UE4 - Fix warning

diff --git a/Engine/Source/Developer/ShaderFormatVectorVM/Private/VectorVMBackend.h b/Engine/Source/Developer/ShaderFormatVectorVM/Private/VectorVMBackend.h
--- a/Engine/Source/Developer/ShaderFormatVectorVM/Private/VectorVMBackend.h
+++ b/Engine/Source/Developer/ShaderFormatVectorVM/Private/VectorVMBackend.h
@@ -2,6 +2,7 @@
 
 #pragma once
 
+#include "CoreMinimal.h"
 #include "hlslcc.h"
 #include "LanguageSpec.h"
 #include "VectorVM.h"
