Copying //UE4/Orion-Staging to //UE4/Main (Source: //Orion/Dev-General @ 3606378)

diff --git a/Engine/Source/Developer/HierarchicalLODUtilities/Public/HierarchicalLODUtilities.h b/Engine/Source/Developer/HierarchicalLODUtilities/Public/HierarchicalLODUtilities.h
--- a/Engine/Source/Developer/HierarchicalLODUtilities/Public/HierarchicalLODUtilities.h
+++ b/Engine/Source/Developer/HierarchicalLODUtilities/Public/HierarchicalLODUtilities.h
@@ -28,7 +28,7 @@ public:
 	virtual float CalculateDrawDistanceFromScreenSize(const float SphereRadius, const float ScreenSize, const FMatrix& ProjectionMatrix) override;
 	virtual UPackage* CreateOrRetrieveLevelHLODPackage(ULevel* InLevel) override;
 	virtual bool BuildStaticMeshForLODActor(ALODActor* LODActor, UPackage* AssetsOuter, const FHierarchicalSimplification& LODSetup) override;
-	virtual bool ShouldGenerateCluster(AActor* Actor) override;
+	virtual EClusterGenerationError ShouldGenerateCluster(AActor* Actor) override;
 	virtual ALODActor* GetParentLODActor(const AActor* InActor) override;
 	virtual void DestroyCluster(ALODActor* InActor) override;
 	virtual void DestroyClusterData(ALODActor* InActor) override;
diff --git a/Engine/Source/Developer/HierarchicalLODUtilities/Public/IHierarchicalLODUtilities.h b/Engine/Source/Developer/HierarchicalLODUtilities/Public/IHierarchicalLODUtilities.h
--- a/Engine/Source/Developer/HierarchicalLODUtilities/Public/IHierarchicalLODUtilities.h
+++ b/Engine/Source/Developer/HierarchicalLODUtilities/Public/IHierarchicalLODUtilities.h
@@ -14,6 +14,23 @@ class ULevel;
 class UStaticMeshComponent;
 struct FHierarchicalSimplification;
 
+enum class EClusterGenerationError : uint32
+{
+	None = 0,
+	ValidActor = 1 << 1,
+	InvalidActor = 1 << 2,
+	ActorHiddenInGame = 1 << 3,
+	ExcludedActor = 1 << 4,
+	LODActor = 1 << 5,
+	ActorTooSmall = 1 << 6,
+	AlreadyClustered = 1 << 7,
+	ComponentHiddenInGame = 1 << 8,
+	MoveableComponent = 1 << 9,
+	ExcludedComponent = 1 << 10
+};
+
+ENUM_CLASS_FLAGS(EClusterGenerationError);
+
 /**
  * IHierarchicalLODUtilities module interface
  */
@@ -59,9 +76,9 @@ public:
 	* Returns whether or not the given actor is eligible for creating a HLOD cluster creation
 	*
 	* @param Actor - Actor to check for if it is eligible for cluster generation
-	* @return bool
+	* @return EClusterGenerationError
 	*/
-	virtual bool ShouldGenerateCluster(AActor* Actor) = 0;
+	virtual EClusterGenerationError ShouldGenerateCluster(AActor* Actor) = 0;
 
 	/** Returns the ALODActor parent for the given InActor, nullptr if none available */
 	virtual ALODActor* GetParentLODActor(const AActor* InActor) = 0;
diff --git a/Engine/Source/Developer/MeshMergeUtilities/Private/MeshMergeUtilities.cpp b/Engine/Source/Developer/MeshMergeUtilities/Private/MeshMergeUtilities.cpp
--- a/Engine/Source/Developer/MeshMergeUtilities/Private/MeshMergeUtilities.cpp
+++ b/Engine/Source/Developer/MeshMergeUtilities/Private/MeshMergeUtilities.cpp
@@ -1000,7 +1000,7 @@ void FMeshMergeUtilities::CreateProxyMesh(const TArray<AActor*>& InActors, const
 		{
 			FSectionInfo& Section = Sections[SectionIndex];
 			const int32 UniqueIndex = UniqueSections.AddUnique(Section);
-			UniqueSectionIndexPerLOD.Add(MeshIndex, TPair<uint32, uint32>(UniqueIndex, SectionIndex));
+			UniqueSectionIndexPerLOD.Add(MeshIndex, TPair<uint32, uint32>(UniqueIndex, Section.MaterialIndex));
 
 			SectionToMesh.Add(UniqueIndex, MeshIndex);
 		}
diff --git a/Engine/Source/Developer/ShaderFormatVectorVM/Private/VectorVMBackend.cpp b/Engine/Source/Developer/ShaderFormatVectorVM/Private/VectorVMBackend.cpp
--- a/Engine/Source/Developer/ShaderFormatVectorVM/Private/VectorVMBackend.cpp
+++ b/Engine/Source/Developer/ShaderFormatVectorVM/Private/VectorVMBackend.cpp
@@ -20,7 +20,6 @@ PRAGMA_POP
 #include "ir.h"
 
 #include "VectorVM.h"
-#include "INiagaraCompiler.h"
 
 #if !PLATFORM_WINDOWS
 #define _strdup strdup
@@ -76,21 +75,26 @@ char* FVectorVMCodeBackend::GenerateCode(exec_list* ir, _mesa_glsl_parse_state*
 		progress = do_copy_propagation_elements(ir) || progress;
 		progress = do_swizzle_swizzle(ir) || progress;
 	} while (progress);
+	//validate_ir_tree(ir, state);
 	vm_debug_dump(ir, state);
 	if (state->error) return nullptr;
 	
 	vm_debug_print("== Branches to selects ==\n");
 	vm_flatten_branches_to_selects(ir, state);
+	//validate_ir_tree(ir, state);
 	vm_debug_dump(ir, state);
 
 	vm_debug_print("== To Single Op ==\n");
 	vm_to_single_op(ir, state);
+	//validate_ir_tree(ir, state);
 	vm_debug_dump(ir, state);
 	if (state->error) return nullptr;
 
 	vm_debug_print("== Scalarize ==\n");
 	vm_scalarize_ops(ir, state);
+	//validate_ir_tree(ir, state);
 	vm_debug_dump(ir, state);
+	//validate_ir_tree(ir, state);
 	if (state->error) return nullptr;
 
 	//99% complete code to remove all matrices from the code and replace them with just swizzled vectors. 
@@ -102,11 +106,14 @@ char* FVectorVMCodeBackend::GenerateCode(exec_list* ir, _mesa_glsl_parse_state*
 
 	vm_debug_print("== Merge Ops ==\n");
 	vm_merge_ops(ir, state);
+//	validate_ir_tree(ir, state);
 	vm_debug_dump(ir, state);
+	//validate_ir_tree(ir, state);
 	if (state->error) return nullptr;
 
 	vm_debug_print("== Propagate non-expressions ==\n");
 	vm_propagate_non_expressions_visitor(ir, state);
+	//validate_ir_tree(ir, state);
 	vm_debug_dump(ir, state);
 	if (state->error) return nullptr;
 
@@ -125,6 +132,8 @@ char* FVectorVMCodeBackend::GenerateCode(exec_list* ir, _mesa_glsl_parse_state*
 	} while (progress);
 	vm_debug_dump(ir, state);
 
+	//validate_ir_tree(ir, state);
+
 	if (state->error) return nullptr;
 
 	vm_gen_bytecode(ir, state, CompilationOutput);
@@ -137,7 +146,8 @@ void FVectorVMLanguageSpec::SetupLanguageIntrinsics(_mesa_glsl_parse_state* Stat
 // 	//We fine if the func is unused entirely but we need to keep the scalar signatures for when we scalarize the call.
 // 	//Maybe we can just keep the wrong sig but still replace the ret val and params?
 // 	make_intrinsic_genType(ir, State, "mad", ir_invalid_opcode, IR_INTRINSIC_FLOAT, 3, 1, 4);
-	make_intrinsic_genType(ir, State, "rand", ir_invalid_opcode, IR_INTRINSIC_FLOAT, 1, 1, 4); 
+	make_intrinsic_genType(ir, State, "rand", ir_invalid_opcode, IR_INTRINSIC_FLOAT, 1, 1, 4);
+	make_intrinsic_genType(ir, State, "rand", ir_invalid_opcode, IR_INTRINSIC_INT, 1, 1, 4);
 	make_intrinsic_genType(ir, State, "Modulo", ir_invalid_opcode, IR_INTRINSIC_FLOAT, 1, 1, 4);
 
 // Dont need all these as we're only using the basic scalar function which we provide the signature for in the usf.
diff --git a/Engine/Source/Developer/ShaderFormatVectorVM/Private/VectorVMBackend.h b/Engine/Source/Developer/ShaderFormatVectorVM/Private/VectorVMBackend.h
--- a/Engine/Source/Developer/ShaderFormatVectorVM/Private/VectorVMBackend.h
+++ b/Engine/Source/Developer/ShaderFormatVectorVM/Private/VectorVMBackend.h
@@ -7,7 +7,7 @@
 #include "LanguageSpec.h"
 #include "VectorVM.h"
 
-struct FNiagaraCompilationOutput;
+struct FVectorVMCompilationOutput;
 
 class FVectorVMLanguageSpec : public ILanguageSpec
 {
@@ -35,7 +35,7 @@ class ir_variable;
 // Generates VectorVM compliant code from IR tokens
 struct FVectorVMCodeBackend : public FCodeBackend
 {
-	FVectorVMCodeBackend(unsigned int InHlslCompileFlags, EHlslCompileTarget InTarget, FNiagaraCompilationOutput& InCompilerOutput) :
+	FVectorVMCodeBackend(unsigned int InHlslCompileFlags, EHlslCompileTarget InTarget, FVectorVMCompilationOutput& InCompilerOutput) :
 		FCodeBackend(InHlslCompileFlags, InTarget), CompilationOutput(InCompilerOutput)
 	{
 	}
@@ -52,7 +52,7 @@ struct FVectorVMCodeBackend : public FCodeBackend
 	//TODO: Do we need to generate a main()?
 	//virtual bool GenerateMain(EHlslShaderFrequency Frequency, const char* EntryPoint, exec_list* Instructions, _mesa_glsl_parse_state* ParseState) override;
 
-	FNiagaraCompilationOutput& CompilationOutput;
+	FVectorVMCompilationOutput& CompilationOutput;
 };
 
 class ir_call;
@@ -77,7 +77,7 @@ void vm_to_single_op(exec_list *ir, _mesa_glsl_parse_state *state);
 void vm_merge_ops(exec_list *ir, _mesa_glsl_parse_state *state);
 void vm_scalarize_ops(exec_list* ir, _mesa_glsl_parse_state* state);
 void vm_propagate_non_expressions_visitor(exec_list* ir, _mesa_glsl_parse_state* state);
-void vm_gen_bytecode(exec_list *ir, _mesa_glsl_parse_state *state, FNiagaraCompilationOutput& InCompOutput);
+void vm_gen_bytecode(exec_list *ir, _mesa_glsl_parse_state *state, FVectorVMCompilationOutput& InCompOutput);
 
 
 //////////////////////////////////////////////////////////////////////////
diff --git a/Engine/Source/Developer/ShaderFormatVectorVM/Private/VectorVMShaderCompiler.cpp b/Engine/Source/Developer/ShaderFormatVectorVM/Private/VectorVMShaderCompiler.cpp
--- a/Engine/Source/Developer/ShaderFormatVectorVM/Private/VectorVMShaderCompiler.cpp
+++ b/Engine/Source/Developer/ShaderFormatVectorVM/Private/VectorVMShaderCompiler.cpp
@@ -9,7 +9,6 @@
 #include "CrossCompilerCommon.h"
 #include "ShaderCompilerCommon.h"
 #include "ShaderPreprocessor.h"
-#include "INiagaraCompiler.h"
 
 #include "VectorVM.h"
 
@@ -21,13 +20,13 @@ DEFINE_LOG_CATEGORY_STATIC(LogVectorVMShaderCompiler, Log, All);
  * @param Input - The input shader code and environment.
  * @param Output - Contains shader compilation results upon return.
  */
-void CompileShader_VectorVM(const FShaderCompilerInput& Input, FShaderCompilerOutput& Output, const FString& WorkingDirectory, uint8 Version)
+bool CompileShader_VectorVM(const FShaderCompilerInput& Input, FShaderCompilerOutput& Output, const FString& WorkingDirectory, uint8 Version)
 {
-
+	return false;
 }
 
 //TODO: Move to this output living in the shader eco-system with the compute shaders too but for now just do things more directly.
-void CompileShader_VectorVM(const FShaderCompilerInput& Input, FShaderCompilerOutput& Output,const FString& WorkingDirectory, uint8 Version, FNiagaraCompilationOutput& NiagaraOutput)
+bool CompileShader_VectorVM(const FShaderCompilerInput& Input, FShaderCompilerOutput& Output,const FString& WorkingDirectory, uint8 Version, FVectorVMCompilationOutput& VMCompilationOutput)
 {
 	FString PreprocessedShader;
 	FShaderCompilerDefinitions AdditionalDefines;
@@ -45,7 +44,7 @@ void CompileShader_VectorVM(const FShaderCompilerInput& Input, FShaderCompilerOu
 	{
 		if (!FFileHelper::LoadFileToString(PreprocessedShader, *Input.VirtualSourceFilePath))
 		{
-			return;
+			return false;
 		}
 
 		// Remove const as we are on debug-only mode
@@ -56,7 +55,7 @@ void CompileShader_VectorVM(const FShaderCompilerInput& Input, FShaderCompilerOu
 		if (!PreprocessShader(PreprocessedShader, Output, Input, AdditionalDefines))
 		{
 			// The preprocessing stage will add any relevant errors.
-			return;
+			return false;
 		}
 	}
 
@@ -92,7 +91,7 @@ void CompileShader_VectorVM(const FShaderCompilerInput& Input, FShaderCompilerOu
 
 	//Is stuff like this needed? What others?
 	uint32 CCFlags = HLSLCC_NoPreprocess;
-	CCFlags |= HLSLCC_PrintAST;
+		//CCFlags |= HLSLCC_PrintAST;
 	//CCFlags |= HLSLCC_UseFullPrecisionInPS;
 
 	//TODO: Do this later when we implement the rest of the shader plumbing stuff.
@@ -112,10 +111,10 @@ void CompileShader_VectorVM(const FShaderCompilerInput& Input, FShaderCompilerOu
 	// Required as we added the RemoveUniformBuffersFromSource() function (the cross-compiler won't be able to interpret comments w/o a preprocessor)
 	//CCFlags &= ~HLSLCC_NoPreprocess;
 
-	FVectorVMCodeBackend VVMBackEnd(CCFlags, HlslCompilerTarget, NiagaraOutput);
+	FVectorVMCodeBackend VVMBackEnd(CCFlags, HlslCompilerTarget, VMCompilationOutput);
 	FVectorVMLanguageSpec VVMLanguageSpec; 
 
-	int32 Result = 0;
+	bool Result = false;
 	FHlslCrossCompilerContext CrossCompilerContext(CCFlags, Frequency, HlslCompilerTarget);
 	if (CrossCompilerContext.Init(TCHAR_TO_ANSI(*Input.VirtualSourceFilePath), &VVMLanguageSpec))
 	{
@@ -128,7 +127,7 @@ void CompileShader_VectorVM(const FShaderCompilerInput& Input, FShaderCompilerOu
 			) ? 1 : 0;
 	}
 
-	NiagaraOutput.Errors = ErrorLog;
+	VMCompilationOutput.Errors = ErrorLog;
 
 	//TODO: Try to get rid of the CompilationOutput and have the vm bytecode life in the shader eco-system as the compute shader version will.
 //		int32 SourceLen = VVMBackEnd.ByteCode.Num();//ShaderSource ? FCStringAnsi::Strlen(ShaderSource) : 0;
@@ -208,4 +207,7 @@ void CompileShader_VectorVM(const FShaderCompilerInput& Input, FShaderCompilerOu
 		UE_LOG(LogVectorVMShaderCompiler, Warning, TEXT("%s"), (const char*)ErrorLog);
 		free(ErrorLog);
 	}
+
+	return Result;
 }
+
diff --git a/Engine/Source/Developer/ShaderFormatVectorVM/Private/ir_vm_flatten_branches_to_selects_visitor.cpp b/Engine/Source/Developer/ShaderFormatVectorVM/Private/ir_vm_flatten_branches_to_selects_visitor.cpp
--- a/Engine/Source/Developer/ShaderFormatVectorVM/Private/ir_vm_flatten_branches_to_selects_visitor.cpp
+++ b/Engine/Source/Developer/ShaderFormatVectorVM/Private/ir_vm_flatten_branches_to_selects_visitor.cpp
@@ -95,6 +95,8 @@ public:
 		// 		if (ir->mode == if_branch)
 		// 			return visit_continue;
 
+		check(base_ir->next && base_ir->prev);
+
 		if (parse_state->error)
 			return visit_stop;
 
@@ -246,6 +248,7 @@ public:
 				{
 					new_deref = replace_assigned_val_with_temp(actual_param->as_dereference());
 					new_assign = new(parent)ir_assignment(actual_param, new_deref);
+					check(actual_param->next && actual_param->prev);
 					actual_param->replace_with(new_deref);
 					curr_assignments->Add(new_assign);
 				}
diff --git a/Engine/Source/Developer/ShaderFormatVectorVM/Private/ir_vm_gen_bytecode_visitor.cpp b/Engine/Source/Developer/ShaderFormatVectorVM/Private/ir_vm_gen_bytecode_visitor.cpp
--- a/Engine/Source/Developer/ShaderFormatVectorVM/Private/ir_vm_gen_bytecode_visitor.cpp
+++ b/Engine/Source/Developer/ShaderFormatVectorVM/Private/ir_vm_gen_bytecode_visitor.cpp
@@ -21,7 +21,6 @@ PRAGMA_POP
 #include "ir.h"
 
 #include "VectorVM.h"
-#include "INiagaraCompiler.h"
 
 //////////////////////////////////////////////////////////////////////////
 
@@ -321,32 +320,26 @@ void BuildExpressionMap()
 	}
 }
 
-//TODO: remove this. Handle internal constants in a way that doesn't require them being written back into the parameters set.
-FNiagaraTypeDefinition GetNiagaraTypeDef(const glsl_type *type)
-{
-	check(type->is_scalar());
-	if (type->is_float())
-	{
-		return FNiagaraTypeDefinition::GetFloatDef();
-	}
-	else if (type->is_integer())
-	{
-		return FNiagaraTypeDefinition::GetIntDef();
-	}
-	else if (type->is_boolean())
-	{
-		return FNiagaraTypeDefinition::GetBoolDef();
-	}
-
-	return FNiagaraTypeDefinition::GetFloatDef();
-}
-
 EVectorVMOp get_special_vm_opcode(ir_function_signature* signature)
 {
 	EVectorVMOp VVMOpCode = EVectorVMOp::done;
 	if (strcmp(signature->function_name(), "rand") == 0)
 	{
-		VVMOpCode = EVectorVMOp::random;
+		unsigned num_operands = 0;
+		foreach_iter(exec_list_iterator, iter, signature->parameters)
+		{
+			ir_variable *const param = (ir_variable *)iter.get();
+			check(param->type->is_scalar());
+			switch (param->type->base_type)
+			{
+			case GLSL_TYPE_FLOAT: VVMOpCode = EVectorVMOp::random; break;
+			case GLSL_TYPE_INT: VVMOpCode = EVectorVMOp::randomi; break;
+			//case GLSL_TYPE_BOOL: v->constant_table_size_bytes += sizeof(int32); break;
+			default: check(0);
+			}
+			++num_operands;
+		}
+		check(num_operands == 1);
 	}
 	else if (strcmp(signature->function_name(), "Modulo") == 0)
 	{
@@ -571,7 +564,7 @@ struct op_base
 		check(component);
 		if (ir_variable* var = component->owner->ir->as_variable())
 		{
-			if (var->mode == ir_var_temporary || var->mode == ir_var_auto)
+			if (component->offset == INDEX_NONE && (var->mode == ir_var_temporary || var->mode == ir_var_auto))
 			{
 				for (unsigned j = 0; j < num_temp_registers; ++j)
 				{
@@ -579,6 +572,9 @@ struct op_base
 					{
 						component->offset = VectorVM::FirstTempRegister + j;
 						registers[j] = component->last_read;
+#if VM_VERBOSE_LOGGING
+						UE_LOG(LogVVMBackend, Log, TEXT("OP:%d | Comp:%p allocated Reg: %d | Last Read: %d |"), op_idx, component, j, component->last_read);
+#endif
 						break;
 					}
 				}
@@ -1068,11 +1064,12 @@ class ir_gen_vvm_visitor : public ir_hierarchical_visitor
 	void *mem_ctx;
 
 	/** Helper class allowing us to feed information about constants back into the compilers representation of constants. This can be removed when "internal" constants are handled differently. */
-	FNiagaraCompilationOutput& CompilationOutput;
+	FVectorVMCompilationOutput& CompilationOutput;
 
 	/** contains a var_info for each variable and constant seen. Each having a tree describing the whole variable and data locations. */
 	hash_table* var_info_table;
 
+	TArray<variable_info*> param_vars;
 	unsigned num_input_components;
 	unsigned num_output_components;
 	unsigned constant_table_size_bytes;
@@ -1087,7 +1084,7 @@ class ir_gen_vvm_visitor : public ir_hierarchical_visitor
 
 	TArray<ir_constant*> seen_constants;
 
-	explicit ir_gen_vvm_visitor(_mesa_glsl_parse_state *in_parse_state, FNiagaraCompilationOutput& InCompilationOutput)
+	explicit ir_gen_vvm_visitor(_mesa_glsl_parse_state *in_parse_state, FVectorVMCompilationOutput& InCompilationOutput)
 		: parse_state(in_parse_state)
 		, mem_ctx(ralloc_context(0))
 		, CompilationOutput(InCompilationOutput)
@@ -1134,6 +1131,7 @@ class ir_gen_vvm_visitor : public ir_hierarchical_visitor
 		ir_gen_vvm_visitor* v = (ir_gen_vvm_visitor*)user_ptr;
 		check(node->is_scalar());
 
+		v->param_vars.AddUnique(node->owner);
 		node->offset = v->constant_table_size_bytes;
 		switch (node->type->base_type)
 		{
@@ -1259,31 +1257,47 @@ class ir_gen_vvm_visitor : public ir_hierarchical_visitor
 			//Internal constants should be written to a header on the bytecode or something and appended to the VM constant table after all the real uniforms.
 			//This is just an implementation quirk of the VM. The System above should know nothing about it.
 
-			//TODO: If we're really gonna do this (for now) maybe better to build up a single struct for all the internal constants and add it as a single uniform.
-			FNiagaraTypeDefinition TypeDef = GetNiagaraTypeDef(constant->type);
+			varinfo = ralloc(mem_ctx, variable_info);
+			hash_table_insert(var_info_table, varinfo, constant);
+			varinfo->Init(mem_ctx, constant, EVectorVMOperandLocation::Constant);
+			int32 offset = CompilationOutput.InternalConstantData.Num();
+			varinfo->root->offset = constant_table_size_bytes + offset;
+			seen_constants.Add(constant);
 
-			FName Name(*FString::Printf(TEXT("internal_const_%u"), seen_constants.Num()));
-			FNiagaraVariable Uni(TypeDef, Name);
 			switch (type->base_type)
 			{
-			case GLSL_TYPE_FLOAT: Uni.SetValue(constant->value.f[0]); break;
-			case GLSL_TYPE_INT: Uni.SetValue(constant->value.i[0]); break;
-			case GLSL_TYPE_BOOL: Uni.SetValue(constant->value.b[0] ? 0xFFFFFFFF : 0x0); break;
+			case GLSL_TYPE_FLOAT: 
+			{
+				CompilationOutput.InternalConstantOffsets.Add(offset);
+				CompilationOutput.InternalConstantTypes.Add(EVectorVMBaseTypes::Float);
+				CompilationOutput.InternalConstantData.AddUninitialized(sizeof(float));
+				*(float*)(CompilationOutput.InternalConstantData.GetData() + offset) = constant->value.f[0];
+				break;
+			}
+			case GLSL_TYPE_INT:
+			{
+				CompilationOutput.InternalConstantOffsets.Add(offset);
+				CompilationOutput.InternalConstantTypes.Add(EVectorVMBaseTypes::Int);
+				CompilationOutput.InternalConstantData.AddUninitialized(sizeof(int32));
+				*(int32*)(CompilationOutput.InternalConstantData.GetData() + offset) = constant->value.i[0];
+				break;
+			}
+			case GLSL_TYPE_BOOL:
+			{
+				CompilationOutput.InternalConstantOffsets.Add(offset);
+				CompilationOutput.InternalConstantTypes.Add(EVectorVMBaseTypes::Bool);
+				CompilationOutput.InternalConstantData.AddUninitialized(sizeof(int32));
+				*(int32*)(CompilationOutput.InternalConstantData.GetData() + offset) = constant->value.b[0] ? 0xFFFFFFFF : 0x0;
+				break;
+			}
 			default: check(0);
 			};
-			CompilationOutput.InternalParameters.SetOrAdd(Uni);
-
-			varinfo = ralloc(mem_ctx, variable_info);
-			hash_table_insert(var_info_table, varinfo, constant);
-			varinfo->Init(mem_ctx, constant, EVectorVMOperandLocation::Constant);
-			assign_uniform_loc(varinfo->root, this);
-			seen_constants.Add(constant);
 		}
 
 		check(varinfo->root->is_scalar());
 		check(varinfo);
 		curr_node = varinfo->root;
-		curr_node->last_read = ordered_ops.Num();
+		curr_node->last_read = ordered_ops.Num() - 1;
 
 		return visit_continue;
 	}
@@ -1301,7 +1315,7 @@ class ir_gen_vvm_visitor : public ir_hierarchical_visitor
 		check(var_info);
 		check(curr_node == nullptr);
 		curr_node = var_info->root->children[index->get_int_component(0)];
-		curr_node->last_read = ordered_ops.Num();
+		curr_node->last_read = ordered_ops.Num() - 1;
 
 		return visit_continue_with_parent;
 	}
@@ -1319,7 +1333,7 @@ class ir_gen_vvm_visitor : public ir_hierarchical_visitor
 		}
 
 		curr_node = var_info->root;
-		curr_node->last_read = ordered_ops.Num();
+		curr_node->last_read = ordered_ops.Num() - 1;
 
 		return visit_continue;
 	}
@@ -1353,7 +1367,7 @@ class ir_gen_vvm_visitor : public ir_hierarchical_visitor
 			curr_node = curr_node->children[swiz_comp];
 		}
 
-		curr_node->last_read = ordered_ops.Num();
+		curr_node->last_read = ordered_ops.Num() - 1;
 
 		//swizzles must be the final entry in a deref chain so we have to have reached a scalar by now.
 		check(curr_node->is_scalar());
@@ -1379,7 +1393,7 @@ class ir_gen_vvm_visitor : public ir_hierarchical_visitor
 		}
 		check(prev_node != curr_node);//We have to find a child to move into.
 
-		curr_node->last_read = ordered_ops.Num();
+		curr_node->last_read = ordered_ops.Num() - 1;
 
 		return visit_continue;
 	}
@@ -1672,7 +1686,7 @@ class ir_gen_vvm_visitor : public ir_hierarchical_visitor
 			func->function_index = CompilationOutput.CalledVMFunctionTable.AddDefaulted();
 			func->sig = call->callee;
 			check(func->function_index != INDEX_NONE);
-			FNiagaraCompilationOutput::FCalledVMFunction& CalledFunc = CompilationOutput.CalledVMFunctionTable[func->function_index];
+			FVectorVMCompilationOutput::FCalledVMFunction& CalledFunc = CompilationOutput.CalledVMFunctionTable[func->function_index];
 			CalledFunc.Name = call->callee->function_name();
 
 			exec_node* sig_param_node = call->callee->parameters.get_head();
@@ -1762,6 +1776,9 @@ class ir_gen_vvm_visitor : public ir_hierarchical_visitor
 		//Allocate temp registers
 		const unsigned num_temp_registers = VectorVM::NumTempRegisters;
 		unsigned *registers = rzalloc_array(mem_ctx, unsigned, num_temp_registers);
+#if VM_VERBOSE_LOGGING
+		UE_LOG(LogVVMBackend, Log, TEXT("\n-------------------------------\nTemporary Allocations\n------------------------------\n"));
+#endif
 		for (int32 i = 0; i < ordered_ops.Num(); ++i)
 		{
 			if (!ordered_ops[i]->finalize_temporary_allocations(num_temp_registers, registers, i))
@@ -1791,20 +1808,84 @@ class ir_gen_vvm_visitor : public ir_hierarchical_visitor
 
 	void DumpOps()
 	{
+#if VM_VERBOSE_LOGGING
 		FString OpStr;
 
+		//Dump the constant table
+		OpStr += TEXT("\n-------------------------------\n");
+		OpStr += TEXT("Constant Table\n");
+		OpStr += TEXT("-------------------------------\n");
+		//First go through all the parameters.
+
+		TFunction<void(variable_info_node*, FString)> GatherConstantTable = [&](variable_info_node* node, FString accumulated_name)
+		{
+			accumulated_name += node->name;
+			if (node->is_scalar())
+			{
+				check(node->offset != INDEX_NONE);
+				OpStr += FString::Printf(TEXT("%d | %s\n"), node->offset, *accumulated_name);
+			}
+			else
+			{
+				for (unsigned i = 0; i < node->num_children; ++i)
+				{
+					GatherConstantTable(node->children[i], accumulated_name);
+				}
+			}
+		};
+
+		for(variable_info* varinfo : param_vars)
+		{
+			GatherConstantTable(varinfo->root, FString());
+		}
+
+		//Then go through all the internal constants.
+		for (int32 i = 0; i < CompilationOutput.InternalConstantOffsets.Num(); ++i)
+		{
+			EVectorVMBaseTypes Type = CompilationOutput.InternalConstantTypes[i];
+			int32 Offset = CompilationOutput.InternalConstantOffsets[i];
+			int32 TableOffset = constant_table_size_bytes + Offset;
+			switch (Type)
+			{
+			case EVectorVMBaseTypes::Float:
+			{
+				float Val = *(float*)(CompilationOutput.InternalConstantData.GetData() + Offset);
+				OpStr += FString::Printf(TEXT("%d | %f\n"), TableOffset, Val);
+			}
+			break;
+			case EVectorVMBaseTypes::Int:
+			{
+				int32 Val = *(int32*)(CompilationOutput.InternalConstantData.GetData() + Offset);
+				OpStr += FString::Printf(TEXT("%d | %d\n"), TableOffset, Val);
+			}
+			break;
+			case EVectorVMBaseTypes::Bool:
+			{
+				int32 Val = *(int32*)(CompilationOutput.InternalConstantData.GetData() + Offset);
+				OpStr += FString::Printf(TEXT("%d | %s\n"), TableOffset, Val == INDEX_NONE ? TEXT("True") : (Val == 0 ? TEXT("False") : TEXT("Invalid") ));
+			}
+			break;
+			}
+		}
+
+		OpStr += TEXT("-------------------------------\n");
+		OpStr += TEXT("Byte Code\n");
+		OpStr += TEXT("-------------------------------\n");
+		
+		//Dump the bytecode		
 		for (int32 op_idx = 0; op_idx < ordered_ops.Num(); ++op_idx)
 		{
-			OpStr += ordered_ops[op_idx]->to_string();
+			OpStr += FString::Printf(TEXT("%d\t| "),op_idx) + ordered_ops[op_idx]->to_string();
 		}
+		
+		OpStr += TEXT("-------------------------------\n");
 
-#if VM_VERBOSE_LOGGING
 		UE_LOG(LogVVMBackend, Log, TEXT("\n%s"), *OpStr);
 #endif
 	}
 
 public:
-	static char* run(exec_list *ir, _mesa_glsl_parse_state *state, FNiagaraCompilationOutput& InCompOutput)
+	static char* run(exec_list *ir, _mesa_glsl_parse_state *state, FVectorVMCompilationOutput& InCompOutput)
 	{
 		//now visit all assignments and gather operations info.
 		ir_gen_vvm_visitor genv(state, InCompOutput);
@@ -1827,7 +1908,7 @@ public:
 };
 
 
-void vm_gen_bytecode(exec_list *ir, _mesa_glsl_parse_state *state, FNiagaraCompilationOutput& InCompOutput)
+void vm_gen_bytecode(exec_list *ir, _mesa_glsl_parse_state *state, FVectorVMCompilationOutput& InCompOutput)
 {
 	BuildExpressionMap();
 
diff --git a/Engine/Source/Developer/ShaderFormatVectorVM/Private/ir_vm_merge_op_visitor.cpp b/Engine/Source/Developer/ShaderFormatVectorVM/Private/ir_vm_merge_op_visitor.cpp
--- a/Engine/Source/Developer/ShaderFormatVectorVM/Private/ir_vm_merge_op_visitor.cpp
+++ b/Engine/Source/Developer/ShaderFormatVectorVM/Private/ir_vm_merge_op_visitor.cpp
@@ -20,7 +20,6 @@ PRAGMA_POP
 #include "ir.h"
 
 #include "VectorVM.h"
-#include "INiagaraCompiler.h"
 
 
 
@@ -100,8 +99,6 @@ public:
 						ir_rvalue* mul_operand = assign->rhs->clone(ralloc_parent(expr), NULL);
 						ir_rvalue* add_operand = equiv_operand == 0 ? expr->operands[1] : expr->operands[0];
 
-						expr->operands[0]->replace_with(mul_operand);
-						expr->operands[1]->replace_with(add_operand);
 						expr->operands[0] = mul_operand;
 						expr->operands[1] = add_operand;
 
@@ -151,8 +148,6 @@ public:
 						ir_rvalue* add_operand = equiv_operand == 0 ? expr->operands[1] : expr->operands[0];
 
 						expr->operation = ir_binop_sub;
-						expr->operands[0]->replace_with(add_operand);
-						expr->operands[1]->replace_with(neg_operand);
 						expr->operands[0] = add_operand;
 						expr->operands[1] = neg_operand;
 
diff --git a/Engine/Source/Developer/ShaderFormatVectorVM/Private/ir_vm_propagate_non_expressions_visitor.cpp b/Engine/Source/Developer/ShaderFormatVectorVM/Private/ir_vm_propagate_non_expressions_visitor.cpp
--- a/Engine/Source/Developer/ShaderFormatVectorVM/Private/ir_vm_propagate_non_expressions_visitor.cpp
+++ b/Engine/Source/Developer/ShaderFormatVectorVM/Private/ir_vm_propagate_non_expressions_visitor.cpp
@@ -21,7 +21,6 @@ PRAGMA_POP
 #include "ir.h"
 
 #include "VectorVM.h"
-#include "INiagaraCompiler.h"
 
 /** Removes any assignments that don't actually map to a VM op but just move some data around. We look for refs and grab the source data direct. */
 class ir_propagate_non_expressions_visitor final : public ir_rvalue_visitor
@@ -113,7 +112,6 @@ public:
 					check(assign->rhs->as_swizzle() || assign->rhs->as_dereference_variable() || assign->rhs->as_constant());
 					check(assign->rhs->type->is_scalar());//All assignments must be scalar at this point!
 					ir_rvalue* new_rval = assign->rhs->clone(ralloc_parent(assign), nullptr);
-					(*rvalue)->replace_with(new_rval);
 					(*rvalue) = new_rval;
 					progress = true;
 				}
@@ -141,6 +139,7 @@ public:
 
 	virtual ir_visitor_status visit_leave(ir_assignment *assign)
 	{
+		check(assign->next && assign->prev);
 		ir_variable* lhs = assign->lhs->variable_referenced();
 		var_info& varinfo = var_info_map.FindOrAdd(lhs);
 		
@@ -353,6 +352,8 @@ public:
 				mv = &MatrixVectorMap.Add(var);
 				void* p = ralloc_parent(var);
 
+				check(var->next && var->prev);
+
 				const glsl_type* vtype = var->type->column_type();
 				const char *name = ralloc_asprintf(p, "%s_%s", var->name,"col0");				
 				mv->v[0] = new(p) ir_variable(vtype, name, var->mode);
diff --git a/Engine/Source/Developer/ShaderFormatVectorVM/Private/ir_vm_scalarize_visitor.cpp b/Engine/Source/Developer/ShaderFormatVectorVM/Private/ir_vm_scalarize_visitor.cpp
--- a/Engine/Source/Developer/ShaderFormatVectorVM/Private/ir_vm_scalarize_visitor.cpp
+++ b/Engine/Source/Developer/ShaderFormatVectorVM/Private/ir_vm_scalarize_visitor.cpp
@@ -258,6 +258,8 @@ class ir_scalarize_visitor2 : public ir_hierarchical_visitor
 		void* parent = ralloc_parent(call->callee);
 		ECallScalarizeMode mode = get_scalarize_mode(call->callee);
 
+		check(call->next && call->prev);
+
 		ir_function_signature* scalar_sig = FindScalarSig(call->callee, true);
 		
 		if (scalar_sig)
@@ -359,6 +361,8 @@ class ir_scalarize_visitor2 : public ir_hierarchical_visitor
 			return visit_stop;
 		}
 
+		check(assign->next && assign->prev);
+
 		void* perm_mem_ctx = ralloc_parent(assign);
 
 		const glsl_type* type = assign->lhs->type;
@@ -415,7 +419,6 @@ class ir_scalarize_visitor2 : public ir_hierarchical_visitor
 				comp_assign->rhs->accept(this);
 				if (curr_rval)
 				{
-					comp_assign->rhs->replace_with(curr_rval);
 					comp_assign->rhs = curr_rval;
 				}
 			}
@@ -561,11 +564,6 @@ class ir_scalarize_visitor2 : public ir_hierarchical_visitor
 
 			if (curr_rval)
 			{
-				if (op + 1 < num_operands)
-					expr->operands[op]->insert_before(curr_rval);
-				else
-					expr->operands[op]->replace_with(curr_rval);
-
 				expr->operands[op] = curr_rval;
 			}
 		}
diff --git a/Engine/Source/Developer/ShaderFormatVectorVM/Private/ir_vm_to_single_op_visitor.cpp b/Engine/Source/Developer/ShaderFormatVectorVM/Private/ir_vm_to_single_op_visitor.cpp
--- a/Engine/Source/Developer/ShaderFormatVectorVM/Private/ir_vm_to_single_op_visitor.cpp
+++ b/Engine/Source/Developer/ShaderFormatVectorVM/Private/ir_vm_to_single_op_visitor.cpp
@@ -52,7 +52,6 @@ public:
 
 		if (replacement)
 		{
-			swiz->val->replace_with(replacement);
 			swiz->val = replacement;
 			replacement = nullptr;
 			progress = true;
@@ -64,6 +63,7 @@ public:
 
 	virtual ir_visitor_status visit_enter(ir_expression* expr)
 	{
+		check(base_ir->next && base_ir->prev);
 		assign_has_expressions = true;
 		expression_stack.Push(expr);
 		for (unsigned i = 0; i < expr->get_num_operands(); ++i)
@@ -72,7 +72,6 @@ public:
 
 			if (replacement)
 			{
-				expr->operands[i]->replace_with(replacement);
 				expr->operands[i] = replacement;
 				replacement = nullptr;
 				progress = true;
@@ -99,6 +98,7 @@ public:
 
 	virtual ir_visitor_status visit_enter(ir_assignment* assign)
 	{
+		check(base_ir->next && base_ir->prev);
 		assign_has_expressions = false;
 		assign->rhs->accept(this);
 
@@ -120,6 +120,7 @@ public:
 
 	virtual ir_visitor_status visit_enter(ir_call* call)
 	{
+		check(base_ir->next && base_ir->prev);
 		assign_has_expressions = true;
 		void *mem_ctx = ralloc_parent(call);
 
@@ -137,6 +138,7 @@ public:
 				base_ir->insert_before(tmp_var);
 				base_ir->insert_before(tmp_assign);
 				ir_rvalue* new_param = new(mem_ctx) ir_dereference_variable(tmp_var);
+				check(param->next && param->prev);
 				param->replace_with(new_param);
 				param = new_param;
 				progress = true;
diff --git a/Engine/Source/Developer/ShaderFormatVectorVM/Private/ir_vm_vec_op_to_scalar_visitor.cpp b/Engine/Source/Developer/ShaderFormatVectorVM/Private/ir_vm_vec_op_to_scalar_visitor.cpp
--- a/Engine/Source/Developer/ShaderFormatVectorVM/Private/ir_vm_vec_op_to_scalar_visitor.cpp
+++ b/Engine/Source/Developer/ShaderFormatVectorVM/Private/ir_vm_vec_op_to_scalar_visitor.cpp
@@ -21,8 +21,6 @@ PRAGMA_POP
 #include "ir.h"
 
 #include "VectorVM.h"
-#include "INiagaraCompiler.h"
-
 
 //Additional helper during scalaization to split any special vector expressions such as dot into their scalar operations.
 class ir_vec_op_to_scalar_visitor : public ir_hierarchical_visitor
@@ -37,6 +35,8 @@ public:
 	{
 		ir_dereference *deref = rval->as_dereference();
 
+		check(base_ir->next && base_ir->prev);
+
 		void* parent = ralloc_parent(base_ir);
 		/* Avoid making a temporary if we don't need to to avoid aliasing. */
 		if (deref &&
@@ -85,6 +85,7 @@ public:
 
 		do_length(result, deref);
 
+		check(call->next && call->prev);
 		call->remove();
 
 		return visit_continue;
@@ -123,6 +124,7 @@ public:
 		case ir_unop_normalize: do_normalize(result, op[0]); break;
 		};
 
+		check(orig_assign->next && orig_assign->prev);
 		orig_assign->remove();
 		return visit_continue;
 	}
@@ -139,6 +141,7 @@ void ir_vec_op_to_scalar_visitor::do_dot(ir_dereference *result, ir_dereference
 {
 	ir_assignment *assign;
 	ir_expression *expr;
+	check(base_ir->next && base_ir->prev);
 	check(a->type == b->type);
 	void* parent = ralloc_parent(base_ir);
 
@@ -164,6 +167,7 @@ void ir_vec_op_to_scalar_visitor::do_dot(ir_dereference *result, ir_dereference
 void ir_vec_op_to_scalar_visitor::do_cross(ir_dereference *result, ir_dereference *a, ir_dereference *b)
 {
 	ir_assignment *assign;
+	check(base_ir->next && base_ir->prev);
 	check(a->type == b->type);
 	check(a->type->is_vector());
 	check(a->type->vector_elements == 3);
@@ -201,6 +205,7 @@ void ir_vec_op_to_scalar_visitor::do_normalize(ir_dereference *result, ir_derefe
 {
 	ir_assignment *assign;
 	ir_expression *expr;
+	check(base_ir->next && base_ir->prev);
 	check(a->type->is_vector());
 
 	void* parent = ralloc_parent(base_ir);
@@ -233,6 +238,7 @@ void ir_vec_op_to_scalar_visitor::do_length(ir_dereference *result, ir_dereferen
 {
 	ir_assignment *assign;
 	ir_expression *expr;
+	check(base_ir->next && base_ir->prev);
 	check(a->type->is_vector());
 
 	void* parent = ralloc_parent(base_ir);
diff --git a/Engine/Source/Developer/ShaderFormatVectorVM/Public/ShaderFormatVectorVM.h b/Engine/Source/Developer/ShaderFormatVectorVM/Public/ShaderFormatVectorVM.h
--- a/Engine/Source/Developer/ShaderFormatVectorVM/Public/ShaderFormatVectorVM.h
+++ b/Engine/Source/Developer/ShaderFormatVectorVM/Public/ShaderFormatVectorVM.h
@@ -2,15 +2,33 @@
 
 #pragma once
 
-#include "NiagaraCommon.h"
-#include "NiagaraParameters.h"
+#include "VectorVM.h"
 
-class UNiagaraNode;
-class UNiagaraGraph;
-class UEdGraphPin;
+/** Data which is generated from the hlsl by the VectorVMBackend and fed back into the */
+struct FVectorVMCompilationOutput
+{
+	FVectorVMCompilationOutput() {}
 
+	TArray<uint8> ByteCode;
 
-extern void SHADERFORMATVECTORVM_API CompileShader_VectorVM(const struct FShaderCompilerInput& Input, struct FShaderCompilerOutput& Output, const class FString& WorkingDirectory, uint8 Version);
+	TArray<int32> InternalConstantOffsets;
+	TArray<uint8> InternalConstantData;
+	TArray<EVectorVMBaseTypes> InternalConstantTypes;
+
+	/** Ordered table of functions actually called by the VM script. */
+	struct FCalledVMFunction
+	{
+		FString Name;
+		TArray<bool> InputParamLocations;
+		int32 NumOutputs;
+		FCalledVMFunction() :NumOutputs(0) {}
+	};
+	TArray<FCalledVMFunction> CalledVMFunctionTable;
+
+	FString Errors;
+};
+
+extern bool SHADERFORMATVECTORVM_API CompileShader_VectorVM(const struct FShaderCompilerInput& Input, struct FShaderCompilerOutput& Output, const class FString& WorkingDirectory, uint8 Version);
 
 //Cheating hack version. To be removed when we add all the plumbing for VVM scripts to be treat like real shaders.
-extern void SHADERFORMATVECTORVM_API CompileShader_VectorVM(const struct FShaderCompilerInput& Input,struct FShaderCompilerOutput& Output,const class FString& WorkingDirectory, uint8 Version, struct FNiagaraCompilationOutput& NiagaraOutput);
+extern bool SHADERFORMATVECTORVM_API CompileShader_VectorVM(const struct FShaderCompilerInput& Input,struct FShaderCompilerOutput& Output,const class FString& WorkingDirectory, uint8 Version, struct FVectorVMCompilationOutput& VMCompilationOutput);
diff --git a/Engine/Source/Developer/ShaderFormatVectorVM/ShaderFormatVectorVM.Build.cs b/Engine/Source/Developer/ShaderFormatVectorVM/ShaderFormatVectorVM.Build.cs
--- a/Engine/Source/Developer/ShaderFormatVectorVM/ShaderFormatVectorVM.Build.cs
+++ b/Engine/Source/Developer/ShaderFormatVectorVM/ShaderFormatVectorVM.Build.cs
@@ -10,7 +10,6 @@ public class ShaderFormatVectorVM : ModuleRules
         PrivateIncludePathModuleNames.AddRange(
             new string[] {
                 "TargetPlatform",
-                "NiagaraEditor"
                 }
             );
 
@@ -23,7 +22,6 @@ public class ShaderFormatVectorVM : ModuleRules
 				"ShaderCompilerCommon",
 				"ShaderPreprocessor",
             	"VectorVM",
-                "Niagara",
                 }
 			);
 
diff --git a/Engine/Source/Developer/SimplygonMeshReduction/Private/SimplygonMeshReduction.cpp b/Engine/Source/Developer/SimplygonMeshReduction/Private/SimplygonMeshReduction.cpp
--- a/Engine/Source/Developer/SimplygonMeshReduction/Private/SimplygonMeshReduction.cpp
+++ b/Engine/Source/Developer/SimplygonMeshReduction/Private/SimplygonMeshReduction.cpp
@@ -133,7 +133,8 @@ public:
 			{
 				if ( Task != nullptr)
 				{
-					Task->EnterProgressFrame(ProgressPercent - PreviousProgress);
+					// @todo: jurre this is temporary fix for UE-48222
+					//Task->EnterProgressFrame(ProgressPercent - PreviousProgress);
 					PreviousProgress = ProgressPercent;
 				}
 			}
diff --git a/Engine/Source/Editor/AnimGraph/Classes/AnimPreviewInstance.h b/Engine/Source/Editor/AnimGraph/Classes/AnimPreviewInstance.h
--- a/Engine/Source/Editor/AnimGraph/Classes/AnimPreviewInstance.h
+++ b/Engine/Source/Editor/AnimGraph/Classes/AnimPreviewInstance.h
@@ -195,8 +195,11 @@ class ANIMGRAPH_API UAnimPreviewInstance : public UAnimSingleNodeInstance
 	//~ Begin UAnimInstance Interface
 	virtual void NativeInitializeAnimation() override;
 	virtual FAnimInstanceProxy* CreateAnimInstanceProxy() override;
+protected:
+	virtual void Montage_Advance(float DeltaTime) override;
 	//~ End UAnimInstance Interface
 
+public:
 	/** Set SkeletalControl Alpha**/
 	void SetSkeletalControlAlpha(float SkeletalControlAlpha);
 
diff --git a/Engine/Source/Editor/AnimGraph/Private/AnimGraphNode_ModifyBone.cpp b/Engine/Source/Editor/AnimGraph/Private/AnimGraphNode_ModifyBone.cpp
--- a/Engine/Source/Editor/AnimGraph/Private/AnimGraphNode_ModifyBone.cpp
+++ b/Engine/Source/Editor/AnimGraph/Private/AnimGraphNode_ModifyBone.cpp
@@ -18,20 +18,24 @@ UAnimGraphNode_ModifyBone::UAnimGraphNode_ModifyBone(const FObjectInitializer& O
 
 void UAnimGraphNode_ModifyBone::ValidateAnimNodeDuringCompilation(USkeleton* ForSkeleton, FCompilerResultsLog& MessageLog)
 {
-	if (ForSkeleton->GetReferenceSkeleton().FindBoneIndex(Node.BoneToModify.BoneName) == INDEX_NONE)
+	// Temporary fix where skeleton is not fully loaded during AnimBP compilation and thus virtual bone name check is invalid UE-39499 (NEED FIX) 
+	if (ForSkeleton && !ForSkeleton->HasAnyFlags(RF_NeedPostLoad))
 	{
-		if (Node.BoneToModify.BoneName == NAME_None)
+		if (ForSkeleton->GetReferenceSkeleton().FindBoneIndex(Node.BoneToModify.BoneName) == INDEX_NONE)
 		{
-			MessageLog.Warning(*LOCTEXT("NoBoneSelectedToModify", "@@ - You must pick a bone to modify").ToString(), this);
-		}
-		else
-		{
-			FFormatNamedArguments Args;
-			Args.Add(TEXT("BoneName"), FText::FromName(Node.BoneToModify.BoneName));
-
-			FText Msg = FText::Format(LOCTEXT("NoBoneFoundToModify", "@@ - Bone {BoneName} not found in Skeleton"), Args);
-
-			MessageLog.Warning(*Msg.ToString(), this);
+			if (Node.BoneToModify.BoneName == NAME_None)
+			{
+				MessageLog.Warning(*LOCTEXT("NoBoneSelectedToModify", "@@ - You must pick a bone to modify").ToString(), this);
+			}
+			else
+			{
+				FFormatNamedArguments Args;
+				Args.Add(TEXT("BoneName"), FText::FromName(Node.BoneToModify.BoneName));
+
+				FText Msg = FText::Format(LOCTEXT("NoBoneFoundToModify", "@@ - Bone {BoneName} not found in Skeleton"), Args);
+
+				MessageLog.Warning(*Msg.ToString(), this);
+			}
 		}
 	}
 
diff --git a/Engine/Source/Editor/AnimGraph/Private/AnimPreviewInstance.cpp b/Engine/Source/Editor/AnimGraph/Private/AnimPreviewInstance.cpp
--- a/Engine/Source/Editor/AnimGraph/Private/AnimPreviewInstance.cpp
+++ b/Engine/Source/Editor/AnimGraph/Private/AnimPreviewInstance.cpp
@@ -459,6 +459,25 @@ void UAnimPreviewInstance::NativeInitializeAnimation()
 	Proxy.RefreshCurveBoneControllers(CurrentAsset);
 }
 
+void UAnimPreviewInstance::Montage_Advance(float DeltaTime)
+{
+	/*
+		We're running in the Animation Editor.
+		Call 'EditorOnly_PreAdvance' on montage instances.
+		So they can do editor specific updates.
+	*/
+	for (int32 InstanceIndex = 0; InstanceIndex < MontageInstances.Num(); InstanceIndex++)
+	{
+		FAnimMontageInstance* const MontageInstance = MontageInstances[InstanceIndex];
+		if (MontageInstance && MontageInstance->IsValid())
+		{
+			MontageInstance->EditorOnly_PreAdvance();
+		}
+	}
+
+	Super::Montage_Advance(DeltaTime);
+}
+
 FAnimNode_ModifyBone* UAnimPreviewInstance::FindModifiedBone(const FName& InBoneName, bool bCurveController/*=false*/)
 {
 	return GetProxyOnGameThread<FAnimPreviewInstanceProxy>().FindModifiedBone(InBoneName, bCurveController);
diff --git a/Engine/Source/Editor/ComponentVisualizers/ComponentVisualizers.Build.cs b/Engine/Source/Editor/ComponentVisualizers/ComponentVisualizers.Build.cs
--- a/Engine/Source/Editor/ComponentVisualizers/ComponentVisualizers.Build.cs
+++ b/Engine/Source/Editor/ComponentVisualizers/ComponentVisualizers.Build.cs
@@ -19,8 +19,9 @@ public class ComponentVisualizers : ModuleRules
 				"UnrealEd",
                 "PropertyEditor",
                 "EditorStyle",
-				"AIModule"
+				"AIModule",
+				"ViewportInteraction"
 			}
-			);
+		);
 	}
 }
diff --git a/Engine/Source/Editor/ComponentVisualizers/Public/Manipulator.h b/Engine/Source/Editor/ComponentVisualizers/Public/Manipulator.h
new file mode 100644
--- /dev/null
+++ b/Engine/Source/Editor/ComponentVisualizers/Public/Manipulator.h
@@ -0,0 +1,47 @@
+// Copyright 1998-2017 Epic Games, Inc. All Rights Reserved.
+
+#pragma once 
+
+#include "CoreMinimal.h"
+#include "UObject/ObjectMacros.h"
+#include "GameFramework/Actor.h"
+#include "ViewportInteractableInterface.h"
+#include "Manipulator.generated.h"
+
+UCLASS()
+class COMPONENTVISUALIZERS_API AManipulator : public AActor, public IViewportInteractableInterface
+{
+	GENERATED_BODY()
+
+public:
+
+	AManipulator();
+
+	// Begin AActor
+	virtual void PostEditMove(bool bFinished) override;
+	virtual bool IsEditorOnly() const final;
+	// End AActor
+
+	// Begin IViewportInteractableInterface
+	virtual void OnPressed(UViewportInteractor* Interactor, const FHitResult& InHitResult, bool& bOutResultedInDrag) override;
+	virtual void OnHover(UViewportInteractor* Interactor) override;
+	virtual void OnHoverEnter(UViewportInteractor* Interactor, const FHitResult& InHitResult) override;
+	virtual void OnHoverLeave(UViewportInteractor* Interactor, const UActorComponent* NewComponent) override;
+	virtual void OnDragRelease(UViewportInteractor* Interactor) override;
+	virtual bool CanBeSelected() override { return true; };
+	// End IViewportInteractableInterface
+
+	/** Set the component that should be moved when the manipulator was moved. */
+	void SetAssociatedComponent(USceneComponent* SceneComponent);
+
+private:
+
+	/** The component to transform when this manipulator was moved. */
+	UPROPERTY()
+	USceneComponent* AssociatedComponent;
+
+	/** Visual representation of this manipulator. */
+	UPROPERTY()
+	class UStaticMeshComponent* StaticMeshComponent;
+
+};
\ No newline at end of file
diff --git a/Engine/Source/Editor/EditorStyle/Private/SlateEditorStyle.cpp b/Engine/Source/Editor/EditorStyle/Private/SlateEditorStyle.cpp
--- a/Engine/Source/Editor/EditorStyle/Private/SlateEditorStyle.cpp
+++ b/Engine/Source/Editor/EditorStyle/Private/SlateEditorStyle.cpp
@@ -4144,6 +4144,28 @@ void FSlateEditorStyle::FStyle::SetupGraphEditorStyles()
 				.SetShadowOffset( FVector2D::ZeroVector )
 				.SetShadowColorAndOpacity( FLinearColor(0.8f,0.8f,0.8f, 0.5) )
 			);
+
+			// Inline Editable Text Block
+			{
+				FTextBlockStyle InlineEditableTextBlockReadOnly = FTextBlockStyle(NormalText)
+					.SetFont(TTF_CORE_FONT("Fonts/Roboto-Regular", 9))
+					.SetColorAndOpacity(FLinearColor(218.0f / 255.0f, 218.0f / 255.0f, 218.0f / 255.0f))
+					.SetShadowOffset(FVector2D::ZeroVector)
+					.SetShadowColorAndOpacity(FLinearColor(0.8f, 0.8f, 0.8f, 0.5));
+
+				FEditableTextBoxStyle InlineEditableTextBlockEditable = FEditableTextBoxStyle()
+					.SetFont(TTF_CORE_FONT("Fonts/Roboto-Regular", 9))
+					.SetBackgroundImageNormal(BOX_BRUSH("Common/TextBox", FMargin(4.0f / 16.0f)))
+					.SetBackgroundImageHovered(BOX_BRUSH("Common/TextBox_Hovered", FMargin(4.0f / 16.0f)))
+					.SetBackgroundImageFocused(BOX_BRUSH("Common/TextBox_Hovered", FMargin(4.0f / 16.0f)))
+					.SetBackgroundImageReadOnly(BOX_BRUSH("Common/TextBox_ReadOnly", FMargin(4.0f / 16.0f)))
+					.SetScrollBarStyle(ScrollBar);
+
+				FInlineEditableTextBlockStyle InlineEditableTextBlockStyle = FInlineEditableTextBlockStyle()
+					.SetTextStyle(InlineEditableTextBlockReadOnly)
+					.SetEditableTextBoxStyle(InlineEditableTextBlockEditable);
+				Set("Graph.Node.InlineEditablePinName", InlineEditableTextBlockStyle);
+			}
 		}
 
 		{
@@ -4695,10 +4717,10 @@ void FSlateEditorStyle::FStyle::SetupLevelEditorStyle()
 		Set( "LevelEditor.FoliageMode.Selected", new IMAGE_BRUSH( "Icons/icon_Mode_Foliage_selected_40x", Icon40x40 ) );
 		Set( "LevelEditor.FoliageMode.Selected.Small", new IMAGE_BRUSH( "Icons/icon_Mode_Foliage_selected_40x", Icon20x20 ) );
 
-		Set("LevelEditor.BspMode", new IMAGE_BRUSH("Icons/icon_Mode_GeoEdit_40px", Icon40x40));
-		Set("LevelEditor.BspMode.Small", new IMAGE_BRUSH("Icons/icon_Mode_GeoEdit_40px", Icon20x20));
-		Set("LevelEditor.BspMode.Selected", new IMAGE_BRUSH("Icons/icon_Mode_GeoEdit-a_40px", Icon40x40));
-		Set("LevelEditor.BspMode.Selected.Small", new IMAGE_BRUSH("Icons/icon_Mode_GeoEdit-a_40px", Icon20x20));
+		Set( "LevelEditor.BspMode", new IMAGE_BRUSH( "Icons/icon_Mode_GeoEdit_40px", Icon40x40 ) );
+		Set( "LevelEditor.BspMode.Small", new IMAGE_BRUSH( "Icons/icon_Mode_GeoEdit_40px", Icon20x20 ) );
+		Set( "LevelEditor.BspMode.Selected", new IMAGE_BRUSH( "Icons/icon_Mode_GeoEdit-a_40px", Icon40x40 ) );
+		Set( "LevelEditor.BspMode.Selected.Small", new IMAGE_BRUSH( "Icons/icon_Mode_GeoEdit-a_40px", Icon20x20 ) );
 
 		Set( "LevelEditor.WorldProperties", new IMAGE_BRUSH( "Icons/icon_worldscript_40x", Icon40x40 ) );
 		Set( "LevelEditor.WorldProperties.Small", new IMAGE_BRUSH( "Icons/icon_worldscript_40x", Icon20x20 ) );
diff --git a/Engine/Source/Editor/GraphEditor/GraphEditor.Build.cs b/Engine/Source/Editor/GraphEditor/GraphEditor.Build.cs
--- a/Engine/Source/Editor/GraphEditor/GraphEditor.Build.cs
+++ b/Engine/Source/Editor/GraphEditor/GraphEditor.Build.cs
@@ -52,13 +52,6 @@ public class GraphEditor : ModuleRules
 			}
 		);
 
-		// TODO: Move niagara code to niagara modules.  This is temporarily necessary to fix the public include issues.
-		PrivateIncludePathModuleNames.AddRange(
-			new string[] {
-				"NiagaraEditor"
-			}
-		);
-
 		DynamicallyLoadedModuleNames.AddRange(
 			new string[] {
 				"ContentBrowser",
diff --git a/Engine/Source/Editor/GraphEditor/Private/KismetPins/SGraphPinEnum.h b/Engine/Source/Editor/GraphEditor/Private/KismetPins/SGraphPinEnum.h
--- a/Engine/Source/Editor/GraphEditor/Private/KismetPins/SGraphPinEnum.h
+++ b/Engine/Source/Editor/GraphEditor/Private/KismetPins/SGraphPinEnum.h
@@ -8,7 +8,7 @@
 #include "SGraphPin.h"
 #include "SGraphPinComboBox.h"
 
-class SGraphPinEnum : public SGraphPin
+class GRAPHEDITOR_API SGraphPinEnum : public SGraphPin
 {
 public:
 	SLATE_BEGIN_ARGS(SGraphPinEnum) {}
diff --git a/Engine/Source/Editor/GraphEditor/Private/NodeFactory.cpp b/Engine/Source/Editor/GraphEditor/Private/NodeFactory.cpp
--- a/Engine/Source/Editor/GraphEditor/Private/NodeFactory.cpp
+++ b/Engine/Source/Editor/GraphEditor/Private/NodeFactory.cpp
@@ -84,8 +84,6 @@
 #include "MaterialGraphConnectionDrawingPolicy.h"
 
 #include "EdGraphUtilities.h"
-#include "EdGraphSchema_Niagara.h"
-
 TSharedPtr<SGraphNode> FNodeFactory::CreateNodeWidget(UEdGraphNode* InNode)
 {
 	check(InNode != NULL);
diff --git a/Engine/Source/Editor/GraphEditor/Private/SGraphPin.cpp b/Engine/Source/Editor/GraphEditor/Private/SGraphPin.cpp
--- a/Engine/Source/Editor/GraphEditor/Private/SGraphPin.cpp
+++ b/Engine/Source/Editor/GraphEditor/Private/SGraphPin.cpp
@@ -181,11 +181,7 @@ void SGraphPin::Construct(const FArguments& InArgs, UEdGraphPin* InPin)
 			.Image(this, &SGraphPin::GetPinStatusIcon)
 		];
 
-	TSharedRef<SWidget> LabelWidget = SNew(STextBlock)
-		.Text(this, &SGraphPin::GetPinLabel)
-		.TextStyle(FEditorStyle::Get(), InArgs._PinLabelStyle)
-		.Visibility(this, &SGraphPin::GetPinLabelVisibility)
-		.ColorAndOpacity(this, &SGraphPin::GetPinTextColor);
+	TSharedRef<SWidget> LabelWidget = GetLabelWidget(InArgs._PinLabelStyle);
 
 	// Create the widget used for the pin body (status indicator, label, and value)
 	TSharedRef<SWrapBox> LabelAndValue =
@@ -306,6 +302,15 @@ TSharedRef<SWidget>	SGraphPin::GetDefaultValueWidget()
 	return SNullWidget::NullWidget;
 }
 
+TSharedRef<SWidget> SGraphPin::GetLabelWidget(const FName& InLabelStyle)
+{
+	return SNew(STextBlock)
+		.Text(this, &SGraphPin::GetPinLabel)
+		.TextStyle(FEditorStyle::Get(), InLabelStyle)
+		.Visibility(this, &SGraphPin::GetPinLabelVisibility)
+		.ColorAndOpacity(this, &SGraphPin::GetPinTextColor);
+}
+
 void SGraphPin::SetIsEditable(TAttribute<bool> InIsEditable)
 {
 	IsEditable = InIsEditable;
diff --git a/Engine/Source/Editor/GraphEditor/Public/SGraphPin.h b/Engine/Source/Editor/GraphEditor/Public/SGraphPin.h
--- a/Engine/Source/Editor/GraphEditor/Public/SGraphPin.h
+++ b/Engine/Source/Editor/GraphEditor/Public/SGraphPin.h
@@ -190,6 +190,9 @@ protected:
 		return bShowLabel ? EVisibility::Visible : EVisibility::Collapsed;
 	}
 
+	/** Get the widget we should put in the label space, which displays the name of the pin.*/
+	virtual TSharedRef<SWidget> GetLabelWidget(const FName& InPinLabelStyle);
+
 	/** @return The brush with which to paint this graph pin's incoming/outgoing bullet point */
 	virtual const FSlateBrush* GetPinIcon() const;
 
diff --git a/Engine/Source/Editor/HierarchicalLODOutliner/Private/HLODOutlinerDragDrop.cpp b/Engine/Source/Editor/HierarchicalLODOutliner/Private/HLODOutlinerDragDrop.cpp
--- a/Engine/Source/Editor/HierarchicalLODOutliner/Private/HLODOutlinerDragDrop.cpp
+++ b/Engine/Source/Editor/HierarchicalLODOutliner/Private/HLODOutlinerDragDrop.cpp
@@ -17,19 +17,17 @@
 
 HLODOutliner::FDragDropPayload::FDragDropPayload()
 {
-	// QQ change this
 	LODActors = TArray<TWeakObjectPtr<AActor>>();
 	StaticMeshActors = TArray<TWeakObjectPtr<AActor>>();
 	bSceneOutliner = false;
 }
 
-bool HLODOutliner::FDragDropPayload::ParseDrag(const FDragDropOperation& Operation)
+EClusterGenerationError HLODOutliner::FDragDropPayload::ParseDrag(const FDragDropOperation& Operation)
 {
-	bool bApplicable = false;
+	EClusterGenerationError ErrorValue = EClusterGenerationError::None;
 
 	if (Operation.IsOfType<FHLODOutlinerDragDropOp>())
 	{
-		bApplicable = true;
 		bSceneOutliner = false;
 
 		const auto& OutlinerOp = static_cast<const FHLODOutlinerDragDropOp&>(Operation);
@@ -43,19 +41,24 @@ bool HLODOutliner::FDragDropPayload::ParseDrag(const FDragDropOperation& Operati
 		{
 			LODActors = OutlinerOp.LODActorOp->Actors;
 		}
+
+		ErrorValue |= EClusterGenerationError::ValidActor;
 	}
 	else if (Operation.IsOfType<FActorDragDropGraphEdOp>())
 	{
 		bSceneOutliner = true;
-		bApplicable = true;
 
+		int32 NumInvalidActors = 0;
 		const auto& ActorOp = static_cast<const FActorDragDropGraphEdOp&>(Operation);
 		for (auto& ActorPtr : ActorOp.Actors)
 		{
 			AActor* Actor = ActorPtr.Get();
 			FHierarchicalLODUtilitiesModule& Module = FModuleManager::LoadModuleChecked<FHierarchicalLODUtilitiesModule>("HierarchicalLODUtilities");
 			IHierarchicalLODUtilities* Utilities = Module.GetUtilities();
-			if (Utilities->ShouldGenerateCluster(Actor))
+
+			EClusterGenerationError ClusterGenerationResult = Utilities->ShouldGenerateCluster(Actor);
+			ErrorValue |= ClusterGenerationResult;
+			if ((ClusterGenerationResult & EClusterGenerationError::ValidActor) != EClusterGenerationError::None)
 			{
 				if (!StaticMeshActors.IsSet())
 				{
@@ -66,14 +69,12 @@ bool HLODOutliner::FDragDropPayload::ParseDrag(const FDragDropOperation& Operati
 			}
 			else
 			{
-				// This means we have invalid static mesh actors in the scene outliner selection (could lead to invalid state)
-				bApplicable = false;
-				break;
+				++NumInvalidActors;
 			}
-		}		
+		}
 	}
 
-	return bApplicable;
+	return ErrorValue;
 }
 
 TSharedPtr<FDragDropOperation> HLODOutliner::CreateDragDropOperation(const TArray<FTreeItemPtr>& InTreeItems)
diff --git a/Engine/Source/Editor/HierarchicalLODOutliner/Private/HLODOutlinerDragDrop.h b/Engine/Source/Editor/HierarchicalLODOutliner/Private/HLODOutlinerDragDrop.h
--- a/Engine/Source/Editor/HierarchicalLODOutliner/Private/HLODOutlinerDragDrop.h
+++ b/Engine/Source/Editor/HierarchicalLODOutliner/Private/HLODOutlinerDragDrop.h
@@ -11,6 +11,7 @@ class AActor;
 class FActorDragDropOp;
 class SWidget;
 struct FSlateBrush;
+enum class EClusterGenerationError : uint32;
 
 namespace HLODOutliner
 {
@@ -47,9 +48,9 @@ namespace HLODOutliner
 
 		/**
 		*	Parse a drag operation into our list of actors and folders
-		*	@return true if the operation is viable for the HLOD outliner to process, false otherwise
+		*	@return Error type None if the operation is viable for the HLOD outliner to process, Error type X otherwise
 		*/
-		bool ParseDrag(const FDragDropOperation& Operation);
+		EClusterGenerationError ParseDrag(const FDragDropOperation& Operation);
 	};
 
 	/** Construct a new Drag and drop operation for a scene outliner */
@@ -62,8 +63,10 @@ namespace HLODOutliner
 		{
 			ToolTip_Compatible,
 			ToolTip_Incompatible,
+			ToolTip_Warning,
 			ToolTip_MultipleSelection_Compatible,
-			ToolTip_MultipleSelection_Incompatible
+			ToolTip_MultipleSelection_Incompatible,
+			ToolTip_MultiSelection_Warning			
 		};
 
 		DRAG_DROP_OPERATOR_TYPE(FHLODOutlinerDragDropOp, FDragDropOperation);
diff --git a/Engine/Source/Editor/HierarchicalLODOutliner/Private/HLODTreeWidgetItem.cpp b/Engine/Source/Editor/HierarchicalLODOutliner/Private/HLODTreeWidgetItem.cpp
--- a/Engine/Source/Editor/HierarchicalLODOutliner/Private/HLODTreeWidgetItem.cpp
+++ b/Engine/Source/Editor/HierarchicalLODOutliner/Private/HLODTreeWidgetItem.cpp
@@ -8,6 +8,7 @@
 #include "HierarchicalLODType.h"
 #include "DragAndDrop/ActorDragDropGraphEdOp.h"
 #include "LODActorItem.h"
+#include "IHierarchicalLODUtilities.h"
 
 #define LOCTEXT_NAMESPACE "HLODTreeWidgetItem"
 
@@ -74,14 +75,77 @@ namespace HLODOutliner
 	{
 		FDragDropPayload DraggedObjects;
 		// Validate now to make sure we don't doing anything we shouldn't
-		if (!DraggedObjects.ParseDrag(*DragDropEvent.GetOperation()))
+		EClusterGenerationError ErrorValue = DraggedObjects.ParseDrag(*DragDropEvent.GetOperation());
+
+		if (ErrorValue != EClusterGenerationError::ValidActor)
 		{
 			// Invalid selection
-			ValidationInfo = FDragValidationInfo(EHierarchicalLODActionType::InvalidAction, FHLODOutlinerDragDropOp::ToolTip_Incompatible, LOCTEXT("InvalidSelection", "Selection contains already clustered Actor(s)"));
-			return FReply::Unhandled();
+			FString ErrorString("Selection contains:");
+			if ((ErrorValue & EClusterGenerationError::AlreadyClustered) != EClusterGenerationError::None)
+			{
+				ErrorString += LOCTEXT("ActorsAlreadyClustered", "\n- Already clustered Actor(s)").ToString();
+			}
+			
+			if ((ErrorValue & EClusterGenerationError::InvalidActor) != EClusterGenerationError::None)
+			{
+				ErrorString += LOCTEXT("InvalidActorsSelected", "\n- Invalid Actor(s)").ToString();
+			}
+			
+			if ((ErrorValue & EClusterGenerationError::ActorHiddenInGame) != EClusterGenerationError::None)
+			{
+				ErrorString += LOCTEXT("ActorsHiddenInGame", "\n- Actor(s) which are hidden In Game").ToString();
+			}
+			
+			if ((ErrorValue & EClusterGenerationError::ExcludedActor) != EClusterGenerationError::None)
+			{
+				ErrorString += LOCTEXT("ActorsExcludedFromHLOD", "\n- Actor(s) which are set to be excluded from HLOD generation").ToString();
+			}
+			
+			if ((ErrorValue & EClusterGenerationError::LODActor) != EClusterGenerationError::None)
+			{
+				ErrorString += LOCTEXT("LODActorsSelected", "\n- LOD Actor(s)").ToString();
+			}
+			
+			if ((ErrorValue & EClusterGenerationError::ActorTooSmall) != EClusterGenerationError::None)
+			{
+				ErrorString += LOCTEXT("InvisibleBoundActors", "\n- Actor(s) with invisible Bounds").ToString();
+			}
+			
+			if ((ErrorValue & EClusterGenerationError::ComponentHiddenInGame) != EClusterGenerationError::None)
+			{
+				ErrorString += LOCTEXT("ComponentsHiddenInGame", "\n- Actor(s) with Components set to be hidden In Game").ToString();
+			}
+			
+			if ((ErrorValue & EClusterGenerationError::MoveableComponent) != EClusterGenerationError::None)
+			{
+				ErrorString += LOCTEXT("MoveableActors", "\n- Actor(s) with Moveable Components").ToString();
+			}
+			
+			if ((ErrorValue & EClusterGenerationError::ExcludedComponent) != EClusterGenerationError::None)
+			{
+				ErrorString += LOCTEXT("ComponentsExcluded", "\n- Actor(s) with Components set to be excluded from HLOD generation").ToString();
+			}
+
+			if ((ErrorValue & EClusterGenerationError::ValidActor) != EClusterGenerationError::None)
+			{
+				// Warning(s)
+				ValidationInfo = DropTarget.ValidateDrop(DraggedObjects);
+				FTextBuilder TextBuilder;
+				TextBuilder.AppendLineFormat(FTextFormat::FromString("{0} - (Warning) {1}"), ValidationInfo.ValidationText, FText::FromString(ErrorString));
+				ValidationInfo.ValidationText = TextBuilder.ToText();
+			}
+			else
+			{
+				// Error
+				ValidationInfo = FDragValidationInfo(EHierarchicalLODActionType::InvalidAction, FHLODOutlinerDragDropOp::ToolTip_Incompatible, FText::FromString(TEXT("(Error) ") + ErrorString));
+				return FReply::Unhandled();
+			}
+		}
+		else
+		{
+			ValidationInfo = DropTarget.ValidateDrop(DraggedObjects);
 		}
 
-		ValidationInfo = DropTarget.ValidateDrop(DraggedObjects);
 
 		if (!ValidationInfo.IsValid())
 		{
diff --git a/Engine/Source/Editor/MovieSceneTools/Private/TrackEditors/AudioTrackEditor.cpp b/Engine/Source/Editor/MovieSceneTools/Private/TrackEditors/AudioTrackEditor.cpp
--- a/Engine/Source/Editor/MovieSceneTools/Private/TrackEditors/AudioTrackEditor.cpp
+++ b/Engine/Source/Editor/MovieSceneTools/Private/TrackEditors/AudioTrackEditor.cpp
@@ -792,6 +792,16 @@ bool FAudioTrackEditor::SupportsType( TSubclassOf<UMovieSceneTrack> Type ) const
 }
 
 
+bool FAudioTrackEditor::SupportsSequence(UMovieSceneSequence* InSequence) const
+{
+	static UClass* LevelSequenceClass = FindObject<UClass>(ANY_PACKAGE, TEXT("LevelSequence"), true);
+	static UClass* WidgetAnimationClass = FindObject<UClass>(ANY_PACKAGE, TEXT("WidgetAnimation"), true);
+	return InSequence != nullptr &&
+		((LevelSequenceClass != nullptr && InSequence->GetClass()->IsChildOf(LevelSequenceClass)) ||
+		(WidgetAnimationClass != nullptr && InSequence->GetClass()->IsChildOf(WidgetAnimationClass)));
+}
+
+
 void CopyInterpSoundTrack(TSharedRef<ISequencer> Sequencer, UInterpTrackSound* MatineeSoundTrack, UMovieSceneAudioTrack* AudioTrack)
 {
 	if (FMatineeImportTools::CopyInterpSoundTrack(MatineeSoundTrack, AudioTrack))
diff --git a/Engine/Source/Editor/MovieSceneTools/Private/TrackEditors/EventTrackEditor.cpp b/Engine/Source/Editor/MovieSceneTools/Private/TrackEditors/EventTrackEditor.cpp
--- a/Engine/Source/Editor/MovieSceneTools/Private/TrackEditors/EventTrackEditor.cpp
+++ b/Engine/Source/Editor/MovieSceneTools/Private/TrackEditors/EventTrackEditor.cpp
@@ -144,6 +144,15 @@ bool FEventTrackEditor::SupportsType(TSubclassOf<UMovieSceneTrack> Type) const
 	return (Type == UMovieSceneEventTrack::StaticClass());
 }
 
+bool  FEventTrackEditor::SupportsSequence(UMovieSceneSequence* InSequence) const
+{
+	static UClass* LevelSequenceClass = FindObject<UClass>(ANY_PACKAGE, TEXT("LevelSequence"), true);
+	static UClass* WidgetAnimationClass = FindObject<UClass>(ANY_PACKAGE, TEXT("WidgetAnimation"), true);
+	return InSequence != nullptr &&
+		((LevelSequenceClass != nullptr && InSequence->GetClass()->IsChildOf(LevelSequenceClass)) ||
+		(WidgetAnimationClass != nullptr && InSequence->GetClass()->IsChildOf(WidgetAnimationClass)));
+}
+
 const FSlateBrush* FEventTrackEditor::GetIconBrush() const
 {
 	return FEditorStyle::GetBrush("Sequencer.Tracks.Event");
diff --git a/Engine/Source/Editor/MovieSceneTools/Private/TrackEditors/EventTrackEditor.h b/Engine/Source/Editor/MovieSceneTools/Private/TrackEditors/EventTrackEditor.h
--- a/Engine/Source/Editor/MovieSceneTools/Private/TrackEditors/EventTrackEditor.h
+++ b/Engine/Source/Editor/MovieSceneTools/Private/TrackEditors/EventTrackEditor.h
@@ -46,6 +46,7 @@ public:
 	virtual void BuildObjectBindingTrackMenu(FMenuBuilder& MenuBuilder, const FGuid& ObjectBinding, const UClass* ObjectClass) override;
 	virtual TSharedRef<ISequencerSection> MakeSectionInterface(UMovieSceneSection& SectionObject, UMovieSceneTrack& Track, FGuid ObjectBinding) override;
 	virtual bool SupportsType( TSubclassOf<UMovieSceneTrack> Type ) const override;
+	virtual bool SupportsSequence(UMovieSceneSequence* InSequence) const override;
 	virtual const FSlateBrush* GetIconBrush() const override;
 	virtual void BuildTrackContextMenu( FMenuBuilder& MenuBuilder, UMovieSceneTrack* Track ) override;
 
diff --git a/Engine/Source/Editor/MovieSceneTools/Private/TrackEditors/MaterialParameterCollectionTrackEditor.cpp b/Engine/Source/Editor/MovieSceneTools/Private/TrackEditors/MaterialParameterCollectionTrackEditor.cpp
--- a/Engine/Source/Editor/MovieSceneTools/Private/TrackEditors/MaterialParameterCollectionTrackEditor.cpp
+++ b/Engine/Source/Editor/MovieSceneTools/Private/TrackEditors/MaterialParameterCollectionTrackEditor.cpp
@@ -84,6 +84,7 @@ void FMaterialParameterCollectionTrackEditor::BuildTrackContextMenu(FMenuBuilder
 	);
 }
 
+
 void FMaterialParameterCollectionTrackEditor::BuildAddTrackMenu(FMenuBuilder& MenuBuilder)
 {
 	auto SubMenuCallback = [this](FMenuBuilder& SubMenuBuilder)
@@ -144,6 +145,15 @@ bool FMaterialParameterCollectionTrackEditor::SupportsType(TSubclassOf<UMovieSce
 	return Type == UMovieSceneMaterialParameterCollectionTrack::StaticClass();
 }
 
+bool FMaterialParameterCollectionTrackEditor::SupportsSequence(UMovieSceneSequence* InSequence) const
+{
+	static UClass* LevelSequenceClass = FindObject<UClass>(ANY_PACKAGE, TEXT("LevelSequence"), true);
+	static UClass* WidgetAnimationClass = FindObject<UClass>(ANY_PACKAGE, TEXT("WidgetAnimation"), true);
+	return InSequence != nullptr &&
+		((LevelSequenceClass != nullptr && InSequence->GetClass()->IsChildOf(LevelSequenceClass)) ||
+		(WidgetAnimationClass != nullptr && InSequence->GetClass()->IsChildOf(WidgetAnimationClass)));
+}
+
 const FSlateBrush* FMaterialParameterCollectionTrackEditor::GetIconBrush() const
 {
 	return FSlateIconFinder::FindIconForClass(UMaterialParameterCollection::StaticClass()).GetIcon();
diff --git a/Engine/Source/Editor/MovieSceneTools/Private/TrackEditors/MaterialParameterCollectionTrackEditor.h b/Engine/Source/Editor/MovieSceneTools/Private/TrackEditors/MaterialParameterCollectionTrackEditor.h
--- a/Engine/Source/Editor/MovieSceneTools/Private/TrackEditors/MaterialParameterCollectionTrackEditor.h
+++ b/Engine/Source/Editor/MovieSceneTools/Private/TrackEditors/MaterialParameterCollectionTrackEditor.h
@@ -30,6 +30,7 @@ public:
 	virtual TSharedPtr<SWidget> BuildOutlinerEditWidget(const FGuid& ObjectBinding, UMovieSceneTrack* Track, const FBuildEditWidgetParams& Params) override;
 	virtual TSharedRef<ISequencerSection> MakeSectionInterface(UMovieSceneSection& SectionObject, UMovieSceneTrack& Track, FGuid ObjectBinding) override;
 	virtual bool SupportsType(TSubclassOf<UMovieSceneTrack> Type) const override;
+	virtual bool SupportsSequence(UMovieSceneSequence* InSequence) const override;
 	virtual void BuildAddTrackMenu(FMenuBuilder& MenuBuilder) override;
 	virtual const FSlateBrush* GetIconBrush() const override;
 	virtual bool HandleAssetAdded(UObject* Asset, const FGuid& TargetObjectGuid) override;
diff --git a/Engine/Source/Editor/PropertyEditor/Private/PropertyHandleImpl.cpp b/Engine/Source/Editor/PropertyEditor/Private/PropertyHandleImpl.cpp
--- a/Engine/Source/Editor/PropertyEditor/Private/PropertyHandleImpl.cpp
+++ b/Engine/Source/Editor/PropertyEditor/Private/PropertyHandleImpl.cpp
@@ -643,6 +643,22 @@ void FPropertyValueImpl::SetOnChildPropertyValueChanged( const FSimpleDelegate&
 	}
 }
 
+void FPropertyValueImpl::SetOnPropertyValuePreChange(const FSimpleDelegate& InOnPropertyValuePreChange)
+{
+	if (PropertyNode.IsValid())
+	{
+		PropertyNode.Pin()->OnPropertyValuePreChange().Add(InOnPropertyValuePreChange);
+	}
+}
+
+void FPropertyValueImpl::SetOnChildPropertyValuePreChange(const FSimpleDelegate& InOnChildPropertyValuePreChange)
+{
+	if (PropertyNode.IsValid())
+	{
+		PropertyNode.Pin()->OnChildPropertyValuePreChange().Add(InOnChildPropertyValuePreChange);
+	}
+}
+
 void FPropertyValueImpl::SetOnRebuildChildren( const FSimpleDelegate& InOnRebuildChildren )
 {
 	if( PropertyNode.IsValid() )
@@ -2235,6 +2251,16 @@ void FPropertyHandleBase::SetOnChildPropertyValueChanged( const FSimpleDelegate&
 	Implementation->SetOnChildPropertyValueChanged( InOnChildPropertyValueChanged );
 }
 
+void FPropertyHandleBase::SetOnPropertyValuePreChange(const FSimpleDelegate& InOnPropertyValuePreChange)
+{
+	Implementation->SetOnPropertyValuePreChange(InOnPropertyValuePreChange);
+}
+
+void FPropertyHandleBase::SetOnChildPropertyValuePreChange(const FSimpleDelegate& InOnChildPropertyValuePreChange)
+{
+	Implementation->SetOnChildPropertyValuePreChange(InOnChildPropertyValuePreChange);
+}
+
 TSharedPtr<FPropertyNode> FPropertyHandleBase::GetPropertyNode() const
 {
 	return Implementation->GetPropertyNode();
diff --git a/Engine/Source/Editor/PropertyEditor/Private/PropertyHandleImpl.h b/Engine/Source/Editor/PropertyEditor/Private/PropertyHandleImpl.h
--- a/Engine/Source/Editor/PropertyEditor/Private/PropertyHandleImpl.h
+++ b/Engine/Source/Editor/PropertyEditor/Private/PropertyHandleImpl.h
@@ -161,6 +161,16 @@ public:
 	 */
 	void SetOnChildPropertyValueChanged( const FSimpleDelegate& InOnChildPropertyValueChanged );
 
+	/**
+	* Sets a delegate to call when the property value is about to change
+	*/
+	void SetOnPropertyValuePreChange(const FSimpleDelegate& InOnPropertyValuePreChange);
+
+	/**
+	* Sets a delegate to call when the propery value of a child is about to change
+	*/
+	void SetOnChildPropertyValuePreChange(const FSimpleDelegate& InOnChildPropertyValuePreChange);
+
 	/**
 	 * Sets a delegate to call when children of the property node must be rebuilt
 	 */
@@ -435,6 +445,8 @@ public:
 	virtual bool IsEditConst() const override;
 	virtual void SetOnPropertyValueChanged( const FSimpleDelegate& InOnPropertyValueChanged ) override;
 	virtual void SetOnChildPropertyValueChanged( const FSimpleDelegate& InOnPropertyValueChanged ) override;
+	virtual void SetOnPropertyValuePreChange(const FSimpleDelegate& InOnPropertyValuePreChange) override;
+	virtual void SetOnChildPropertyValuePreChange(const FSimpleDelegate& InOnPropertyValuePreChange) override;
 	virtual int32 GetIndexInArray() const override;
 	virtual FPropertyAccess::Result GetValueAsFormattedString( FString& OutValue, EPropertyPortFlags PortFlags = PPF_PropertyWindow ) const override;
 	virtual FPropertyAccess::Result GetValueAsDisplayString( FString& OutValue, EPropertyPortFlags PortFlags = PPF_PropertyWindow) const override;
diff --git a/Engine/Source/Editor/PropertyEditor/Private/PropertyNode.cpp b/Engine/Source/Editor/PropertyEditor/Private/PropertyNode.cpp
--- a/Engine/Source/Editor/PropertyEditor/Private/PropertyNode.cpp
+++ b/Engine/Source/Editor/PropertyEditor/Private/PropertyNode.cpp
@@ -2300,6 +2300,9 @@ void FPropertyNode::NotifyPreChange( UProperty* PropertyAboutToChange, FNotifyHo
 			}
 		}
 	}
+
+	// Broadcast the change to any listeners
+	BroadcastPropertyPreChangeDelegates();
 }
 
 void FPropertyNode::NotifyPostChange( FPropertyChangedEvent& InPropertyChangedEvent, class FNotifyHook* InNotifyHook )
@@ -2464,6 +2467,24 @@ void FPropertyNode::BroadcastPropertyChangedDelegates()
 
 }
 
+void FPropertyNode::BroadcastPropertyPreChangeDelegates()
+{
+	PropertyValuePreChangeEvent.Broadcast();
+
+	// Walk through the parents and broadcast
+	FPropertyNode* LocalParentNode = GetParentNode();
+	while (LocalParentNode)
+	{
+		if (LocalParentNode->OnChildPropertyValuePreChange().IsBound())
+		{
+			LocalParentNode->OnChildPropertyValuePreChange().Broadcast();
+		}
+
+		LocalParentNode = LocalParentNode->GetParentNode();
+	}
+
+}
+
 void FPropertyNode::SetOnRebuildChildren( FSimpleDelegate InOnRebuildChildren )
 {
 	OnRebuildChildren = InOnRebuildChildren;
diff --git a/Engine/Source/Editor/PropertyEditor/Private/PropertyNode.h b/Engine/Source/Editor/PropertyEditor/Private/PropertyNode.h
--- a/Engine/Source/Editor/PropertyEditor/Private/PropertyNode.h
+++ b/Engine/Source/Editor/PropertyEditor/Private/PropertyNode.h
@@ -542,6 +542,13 @@ public:
 	/** Broadcasts when a child of this property changes */
 	FPropertyValueChangedEvent& OnChildPropertyValueChanged() { return ChildPropertyValueChangedEvent; }
 
+	/** Broadcasts when a property value changes */
+	DECLARE_EVENT(FPropertyNode, FPropertyValuePreChangeEvent);
+	FPropertyValuePreChangeEvent& OnPropertyValuePreChange() { return PropertyValuePreChangeEvent; }
+
+	/** Broadcasts when a child of this property changes */
+	FPropertyValuePreChangeEvent& OnChildPropertyValuePreChange() { return ChildPropertyValuePreChangeEvent; }
+
 	/**
 	 * Marks window's seem due to filtering flags
 	 * @param InFilterStrings	- List of strings that must be in the property name in order to display
@@ -851,6 +858,13 @@ protected:
 	 */
 	void BroadcastPropertyChangedDelegates();
 
+
+	/**
+	* Helper function for derived members to be able to
+	* broadcast property pre-change notifications
+	*/
+	void BroadcastPropertyPreChangeDelegates();
+
 	/**
 	 * Gets a value tracker for the default of this property in the passed in object
 	 *
@@ -894,6 +908,12 @@ protected:
 	/** Called when this node's children are rebuilt */
 	FSimpleDelegate OnRebuildChildren;
 
+	/** Called when this node's property value is about to change (called during NotifyPreChange) */
+	FPropertyValuePreChangeEvent PropertyValuePreChangeEvent;
+
+	/** Called when a child's property value is about to change */
+	FPropertyValuePreChangeEvent ChildPropertyValuePreChangeEvent;
+
 	/** Called when this node's property value has changed (called during NotifyPostChange) */
 	FPropertyValueChangedEvent PropertyValueChangedEvent;
 	
diff --git a/Engine/Source/Editor/PropertyEditor/Public/PropertyHandle.h b/Engine/Source/Editor/PropertyEditor/Public/PropertyHandle.h
--- a/Engine/Source/Editor/PropertyEditor/Public/PropertyHandle.h
+++ b/Engine/Source/Editor/PropertyEditor/Public/PropertyHandle.h
@@ -210,6 +210,20 @@ public:
 	 */
 	virtual void SetOnChildPropertyValueChanged( const FSimpleDelegate& InOnChildPropertyValueChanged ) = 0;
 
+	/**
+	* Sets a delegate to call when the value of the property is about to be changed
+	*
+	* @param InOnPropertyValuePreChange	The delegate to call
+	*/
+	virtual void SetOnPropertyValuePreChange(const FSimpleDelegate& InOnPropertyValuePreChange) = 0;
+
+	/**
+	* Sets a delegate to call when the value of the property of a child is about to be changed
+	*
+	* @param InOnChildPropertyValuePreChange	The delegate to call
+	*/
+	virtual void SetOnChildPropertyValuePreChange(const FSimpleDelegate& InOnChildPropertyValuePreChange) = 0;
+
 	/**
 	 * Gets the typed value of a property.  
 	 * If the property does not support the value type FPropertyAccess::Fail is returned
diff --git a/Engine/Source/Editor/Sequencer/Private/SSequencer.cpp b/Engine/Source/Editor/Sequencer/Private/SSequencer.cpp
--- a/Engine/Source/Editor/Sequencer/Private/SSequencer.cpp
+++ b/Engine/Source/Editor/Sequencer/Private/SSequencer.cpp
@@ -2297,7 +2297,7 @@ void SSequencer::OnSequenceInstanceActivated( FMovieSceneSequenceIDRef ActiveIns
 	if ( Sequencer.IsValid() )
 	{
 		UMovieScene* MovieScene = Sequencer->GetFocusedMovieSceneSequence()->GetMovieScene();
-		if ( MovieScene->GetFixedFrameInterval() == 0 )
+		if ( MovieScene->GetForceFixedFrameIntervalPlayback() && MovieScene->GetFixedFrameInterval() == 0 )
 		{
 			MovieScene->Modify();
 			MovieScene->SetFixedFrameInterval(Settings->GetTimeSnapInterval());
diff --git a/Engine/Source/Editor/Sequencer/Private/Sequencer.cpp b/Engine/Source/Editor/Sequencer/Private/Sequencer.cpp
--- a/Engine/Source/Editor/Sequencer/Private/Sequencer.cpp
+++ b/Engine/Source/Editor/Sequencer/Private/Sequencer.cpp
@@ -1720,11 +1720,18 @@ void FSequencer::ScrollIntoView(float InLocalTime)
 		if (RangeOffset != 0.f)
 		{
 			TRange<float> WorkingRange = GetClampRange();
-			if (TargetViewRange.GetLowerBoundValue() + RangeOffset > WorkingRange.GetLowerBoundValue() &&
-				TargetViewRange.GetUpperBoundValue() + RangeOffset < WorkingRange.GetUpperBoundValue())
+
+			// Adjust the offset so that the target range will be within the working range.
+			if (TargetViewRange.GetLowerBoundValue() + RangeOffset < WorkingRange.GetLowerBoundValue())
+			{
+				RangeOffset = WorkingRange.GetLowerBoundValue() - TargetViewRange.GetLowerBoundValue();
+			}
+			else if (TargetViewRange.GetUpperBoundValue() + RangeOffset > WorkingRange.GetUpperBoundValue())
 			{
-				SetViewRange(TRange<float>(TargetViewRange.GetLowerBoundValue() + RangeOffset, TargetViewRange.GetUpperBoundValue() + RangeOffset), EViewRangeInterpolation::Immediate);
+				RangeOffset = WorkingRange.GetUpperBoundValue() - TargetViewRange.GetUpperBoundValue();
 			}
+
+			SetViewRange(TRange<float>(TargetViewRange.GetLowerBoundValue() + RangeOffset, TargetViewRange.GetUpperBoundValue() + RangeOffset), EViewRangeInterpolation::Immediate);
 		}
 	}
 }
@@ -3935,6 +3942,15 @@ void FSequencer::OnSelectedOutlinerNodesChanged()
 
 	OnSelectionChangedObjectGuidsDelegate.Broadcast(Selection.GetBoundObjectsGuids());
 	OnSelectionChangedTracksDelegate.Broadcast(Selection.GetSelectedTracks());
+	TArray<UMovieSceneSection*> SelectedSections;
+	for (TWeakObjectPtr<UMovieSceneSection> SelectedSectionPtr : Selection.GetSelectedSections())
+	{
+		if (SelectedSectionPtr.IsValid())
+		{
+			SelectedSections.Add(SelectedSectionPtr.Get());
+		}
+	}
+	OnSelectionChangedSectionsDelegate.Broadcast(SelectedSections);
 }
 
 
@@ -4217,6 +4233,21 @@ FSequencerSelectionPreview& FSequencer::GetSelectionPreview()
 	return SelectionPreview;
 }
 
+void FSequencer::GetSelectedTracks(TArray<UMovieSceneTrack*>& OutSelectedTracks)
+{
+	OutSelectedTracks.Append(Selection.GetSelectedTracks());
+}
+
+void FSequencer::GetSelectedSections(TArray<UMovieSceneSection*>& OutSelectedSections)
+{
+	for (TWeakObjectPtr<UMovieSceneSection> SelectedSection : Selection.GetSelectedSections())
+	{
+		if (SelectedSection.IsValid())
+		{
+			OutSelectedSections.Add(SelectedSection.Get());
+		}
+	}
+}
 
 void FSequencer::SelectObject(FGuid ObjectBinding)
 {
@@ -4228,6 +4259,28 @@ void FSequencer::SelectObject(FGuid ObjectBinding)
 	}
 }
 
+void FSequencer::SelectTrack(UMovieSceneTrack* Track)
+{
+	for (TSharedRef<FSequencerDisplayNode> Node : NodeTree->GetAllNodes())
+	{
+		if (Node->GetType() == ESequencerNode::Track)
+		{
+			TSharedRef<FSequencerTrackNode> TrackNode = StaticCastSharedRef<FSequencerTrackNode>(Node);
+			UMovieSceneTrack* TrackForNode = TrackNode->GetTrack();
+			if (TrackForNode == Track)
+			{
+				Selection.AddToSelection(Node);
+				break;
+			}
+		}
+	}
+}
+
+void FSequencer::SelectSection(UMovieSceneSection* Section)
+{
+	Selection.AddToSelection(Section);
+}
+
 void FSequencer::SelectByPropertyPaths(const TArray<FString>& InPropertyPaths)
 {
 	TArray<TSharedRef<FSequencerDisplayNode>> NodesToSelect;
@@ -4259,6 +4312,11 @@ void FSequencer::SelectByPropertyPaths(const TArray<FString>& InPropertyPaths)
 	}
 }
 
+void FSequencer::EmptySelection()
+{
+	Selection.Empty();
+}
+
 float FSequencer::GetOverlayFadeCurve() const
 {
 	return OverlayCurve.GetLerp();
diff --git a/Engine/Source/Editor/Sequencer/Private/Sequencer.h b/Engine/Source/Editor/Sequencer/Private/Sequencer.h
--- a/Engine/Source/Editor/Sequencer/Private/Sequencer.h
+++ b/Engine/Source/Editor/Sequencer/Private/Sequencer.h
@@ -623,8 +623,13 @@ public:
 	virtual void KeyProperty(FKeyPropertyParams KeyPropertyParams) override;
 	virtual FSequencerSelection& GetSelection() override;
 	virtual FSequencerSelectionPreview& GetSelectionPreview() override;
+	virtual void GetSelectedTracks(TArray<UMovieSceneTrack*>& OutSelectedTracks) override;
+	virtual void GetSelectedSections(TArray<UMovieSceneSection*>& OutSelectedSections) override;
 	virtual void SelectObject(FGuid ObjectBinding) override;
+	virtual void SelectTrack(UMovieSceneTrack* Track) override;
+	virtual void SelectSection(UMovieSceneSection* Section) override;
 	virtual void SelectByPropertyPaths(const TArray<FString>& InPropertyPaths) override;
+	virtual void EmptySelection() override;
 	virtual FOnGlobalTimeChanged& OnGlobalTimeChanged() override { return OnGlobalTimeChangedDelegate; }
 	virtual FOnBeginScrubbingEvent& OnBeginScrubbingEvent() override { return OnBeginScrubbingDelegate; }
 	virtual FOnEndScrubbingEvent& OnEndScrubbingEvent() override { return OnEndScrubbingDelegate; }
@@ -632,6 +637,7 @@ public:
 	virtual FOnMovieSceneBindingsChanged& OnMovieSceneBindingsChanged() override { return OnMovieSceneBindingsChangedDelegate; }
 	virtual FOnSelectionChangedObjectGuids& GetSelectionChangedObjectGuids() override { return OnSelectionChangedObjectGuidsDelegate; }
 	virtual FOnSelectionChangedTracks& GetSelectionChangedTracks() override { return OnSelectionChangedTracksDelegate; }
+	virtual FOnSelectionChangedSections& GetSelectionChangedSections() override { return OnSelectionChangedSectionsDelegate; }
 	virtual FGuid CreateBinding(UObject& InObject, const FString& InName) override;
 	virtual UObject* GetPlaybackContext() const override;
 	virtual TArray<UObject*> GetEventContexts() const override;
@@ -1006,6 +1012,9 @@ private:
 	/** A delegate which is called any time the sequencer selection changes. */
 	FOnSelectionChangedTracks OnSelectionChangedTracksDelegate;
 
+	/** A delegate which is called any time the sequencer selection changes. */
+	FOnSelectionChangedSections OnSelectionChangedSectionsDelegate;
+
 	FOnActorAddedToSequencer OnActorAddedToSequencerEvent;
 	FOnCameraCut OnCameraCutEvent;
 	FOnPreSave OnPreSaveEvent;
diff --git a/Engine/Source/Editor/Sequencer/Private/SequencerTimeSliderController.cpp b/Engine/Source/Editor/Sequencer/Private/SequencerTimeSliderController.cpp
--- a/Engine/Source/Editor/Sequencer/Private/SequencerTimeSliderController.cpp
+++ b/Engine/Source/Editor/Sequencer/Private/SequencerTimeSliderController.cpp
@@ -137,6 +137,12 @@ struct FDrawTickArgs
 
 void FSequencerTimeSliderController::DrawTicks( FSlateWindowElementList& OutDrawElements, const struct FScrubRangeToScreen& RangeToScreen, FDrawTickArgs& InArgs ) const
 {
+	// The math here breaks down when pixels per input is near zero or zero, so just skip drawing ticks to avoid an infinite loop.
+	if (FMath::IsNearlyZero(RangeToScreen.PixelsPerInput))
+	{
+		return;
+	}
+
 	float MinDisplayTickSpacing = ScrubConstants::MinDisplayTickSpacing;
 	if (SequencerSnapValues::IsTimeSnapIntervalFrameRate(TimeSliderArgs.TimeSnapInterval.Get()) && TimeSliderArgs.TimeSnapInterval.Get())
 	{
diff --git a/Engine/Source/Editor/Sequencer/Public/ISequencer.h b/Engine/Source/Editor/Sequencer/Public/ISequencer.h
--- a/Engine/Source/Editor/Sequencer/Public/ISequencer.h
+++ b/Engine/Source/Editor/Sequencer/Public/ISequencer.h
@@ -136,6 +136,8 @@ public:
 
 	DECLARE_MULTICAST_DELEGATE_OneParam(FOnSelectionChangedTracks, TArray<UMovieSceneTrack*> /*Tracks*/)
 
+	DECLARE_MULTICAST_DELEGATE_OneParam(FOnSelectionChangedSections, TArray<UMovieSceneSection*> /*Sections*/)
+
 	DECLARE_MULTICAST_DELEGATE_TwoParams(FOnActorAddedToSequencer, AActor*, const FGuid);
 
 public:
@@ -383,12 +385,27 @@ public:
 	virtual FSequencerSelection& GetSelection() = 0;
 	virtual FSequencerSelectionPreview& GetSelectionPreview() = 0;
 
+	/** Gets the currently selected tracks. */
+	virtual void GetSelectedTracks(TArray<UMovieSceneTrack*>& OutSelectedTracks) = 0;
+
+	/** Gets the currently selected sections. */
+	virtual void GetSelectedSections(TArray<UMovieSceneSection*>& OutSelectedSections) = 0;
+
 	/** Selects an object by GUID */
 	virtual void SelectObject(FGuid ObjectBinding) = 0;
 
+	/** Selects a track */
+	virtual void SelectTrack(UMovieSceneTrack* Track) = 0;
+
+	/** Selects a section */
+	virtual void SelectSection(UMovieSceneSection* Section) = 0;
+
 	/** Selects property tracks by property path */
 	virtual void SelectByPropertyPaths(const TArray<FString>& InPropertyPaths) = 0;
 
+	/** Empties the current selection. */
+	virtual void EmptySelection() = 0;
+
 	/** Gets a multicast delegate which is executed whenever the global time changes. */
 	virtual FOnGlobalTimeChanged& OnGlobalTimeChanged() = 0;
 
@@ -410,6 +427,9 @@ public:
 	/** Gets a multicast delegate with an array of UMovieSceneTracks which is called when the outliner node selection changes. */
 	virtual FOnSelectionChangedTracks& GetSelectionChangedTracks() = 0;
 
+	/** Gets a multicast delegate with an array of UMovieSceneSections which is called when the outliner node selection changes. */
+	virtual FOnSelectionChangedSections& GetSelectionChangedSections() = 0;
+
 	/** @return a numeric type interface that will parse and display numbers as frames and times correctly */
 	virtual TSharedRef<INumericTypeInterface<float>> GetNumericTypeInterface() = 0;
 
diff --git a/Engine/Source/Editor/UnrealEd/Private/EditorEngine.cpp b/Engine/Source/Editor/UnrealEd/Private/EditorEngine.cpp
--- a/Engine/Source/Editor/UnrealEd/Private/EditorEngine.cpp
+++ b/Engine/Source/Editor/UnrealEd/Private/EditorEngine.cpp
@@ -918,7 +918,6 @@ void UEditorEngine::Init(IEngineLoop* InEngineLoop)
 			TEXT("TreeMap"),
 			TEXT("SizeMap"),
 			TEXT("MergeActors"),
-			TEXT("NiagaraEditor"),
 			TEXT("InputBindingEditor"),
 			TEXT("AudioEditor")
 		};
diff --git a/Engine/Source/Editor/UnrealEd/Private/SColorGradientEditor.cpp b/Engine/Source/Editor/UnrealEd/Private/SColorGradientEditor.cpp
--- a/Engine/Source/Editor/UnrealEd/Private/SColorGradientEditor.cpp
+++ b/Engine/Source/Editor/UnrealEd/Private/SColorGradientEditor.cpp
@@ -603,8 +603,10 @@ void SColorGradientEditor::OpenGradientStopColorPicker()
 void SColorGradientEditor::OnSelectedStopColorChanged( FLinearColor InNewColor )
 {
 	FScopedTransaction ColorChange( LOCTEXT("ChangeGradientStopColor", "Change Gradient Stop Color") );
-	SelectedStop.SetColor( InNewColor, *CurveOwner );
 	CurveOwner->ModifyOwner();
+	SelectedStop.SetColor( InNewColor, *CurveOwner );
+	TArray<FRichCurveEditInfo> ChangedCurves{ CurveOwner->GetCurves()[0], CurveOwner->GetCurves()[1], CurveOwner->GetCurves()[2] };
+	CurveOwner->OnCurveChanged(ChangedCurves);
 
 	// Set the the last edited color.  The next time a new stop is added we'll use this value
 	LastModifiedColor.R = InNewColor.R;
@@ -614,8 +616,10 @@ void SColorGradientEditor::OnSelectedStopColorChanged( FLinearColor InNewColor )
 
 void SColorGradientEditor::OnCancelSelectedStopColorChange( FLinearColor PreviousColor )
 {
-	SelectedStop.SetColor( PreviousColor, *CurveOwner );
 	CurveOwner->ModifyOwner();
+	SelectedStop.SetColor( PreviousColor, *CurveOwner );
+	TArray<FRichCurveEditInfo> ChangedCurves{ CurveOwner->GetCurves()[0], CurveOwner->GetCurves()[1], CurveOwner->GetCurves()[2] };
+	CurveOwner->OnCurveChanged(ChangedCurves);
 }
 
 void SColorGradientEditor::OnBeginChangeAlphaValue()
@@ -643,6 +647,8 @@ void SColorGradientEditor::OnAlphaValueChanged( float NewValue )
 	{
 		// RGB is ignored in this case
 		SelectedStop.SetColor( FLinearColor( 0,0,0, NewValue ), *CurveOwner );
+		TArray<FRichCurveEditInfo> ChangedCurves{ CurveOwner->GetCurves()[3] };
+		CurveOwner->OnCurveChanged(ChangedCurves);
 	}
 }
 
@@ -653,12 +659,15 @@ void SColorGradientEditor::OnAlphaValueCommitted( float NewValue, ETextCommit::T
 		// Value was typed in, no transaction is active
 		FScopedTransaction ChangeAlphaTransaction( LOCTEXT("ChangeGradientStopAlpha", "Change Gradient Stop Alpha") );
 		CurveOwner->ModifyOwner();
-
 		SelectedStop.SetColor( FLinearColor( 0,0,0, NewValue ), *CurveOwner );
+		TArray<FRichCurveEditInfo> ChangedCurves{ CurveOwner->GetCurves()[3] };
+		CurveOwner->OnCurveChanged(ChangedCurves);
 	}
 	else
 	{
 		SelectedStop.SetColor( FLinearColor( 0,0,0, NewValue ), *CurveOwner );
+		TArray<FRichCurveEditInfo> ChangedCurves{ CurveOwner->GetCurves()[3] };
+		CurveOwner->OnCurveChanged(ChangedCurves);
 	}
 
 	// Set the alpha of the last edited color.  The next time a new alpha stop is added we'll use this value
@@ -678,8 +687,8 @@ void SColorGradientEditor::OnSetGradientStopTimeFromPopup( const FText& NewText,
 
 		FScopedTransaction Transaction( LOCTEXT("ChangeGradientStopTime", "Change Gradient Stop Time" ) );
 		CurveOwner->ModifyOwner();
-
 		SelectedStop.SetTime( NewTime, *CurveOwner );
+		CurveOwner->OnCurveChanged(CurveOwner->GetCurves());
 	}
 }
 
@@ -849,6 +858,8 @@ void SColorGradientEditor::DeleteStop( const FGradientStopMark& InMark )
 		GreenCurve->DeleteKey( InMark.GreenKeyHandle );
 		BlueCurve->DeleteKey( InMark.BlueKeyHandle );
 	}
+
+	CurveOwner->OnCurveChanged(CurveOwner->GetCurves());
 }
 
 FGradientStopMark SColorGradientEditor::AddStop( const FVector2D& Position, const FGeometry& MyGeometry, bool bColorStop )
@@ -884,13 +895,16 @@ FGradientStopMark SColorGradientEditor::AddStop( const FVector2D& Position, cons
 		NewStop.AlphaKeyHandle = AlphaCurve->AddKey( NewStopTime, LastModifiedColor.A );
 	}
 
+	CurveOwner->OnCurveChanged(CurveOwner->GetCurves());
+
 	return NewStop;
 }
 
 void SColorGradientEditor::MoveStop( FGradientStopMark& Mark, float NewTime )
 {
-	Mark.SetTime( NewTime, *CurveOwner );
 	CurveOwner->ModifyOwner();
+	Mark.SetTime( NewTime, *CurveOwner );
+	CurveOwner->OnCurveChanged(CurveOwner->GetCurves());
 }
 
 
diff --git a/Engine/Source/Editor/UnrealEd/Private/SCurveEditor.cpp b/Engine/Source/Editor/UnrealEd/Private/SCurveEditor.cpp
--- a/Engine/Source/Editor/UnrealEd/Private/SCurveEditor.cpp
+++ b/Engine/Source/Editor/UnrealEd/Private/SCurveEditor.cpp
@@ -97,6 +97,8 @@ void SCurveEditor::Construct(const FArguments& InArgs)
 	bShowCurveSelector = InArgs._ShowCurveSelector;
 	bDrawInputGridNumbers = InArgs._ShowInputGridNumbers;
 	bDrawOutputGridNumbers = InArgs._ShowOutputGridNumbers;
+	bAreCurvesVisible = InArgs._AreCurvesVisible;
+	SetAreCurvesVisibleHandler = InArgs._OnSetAreCurvesVisible;
 
 	OnCreateAsset = InArgs._OnCreateAsset;
 
@@ -1255,8 +1257,11 @@ void SCurveEditor::SetCurveOwner(FCurveOwnerInterface* InCurveOwner, bool bCanEd
 	CurveOwner = InCurveOwner;
 	bCanEditTrack = bCanEdit;
 
+	if (bAreCurvesVisible.IsBound() == false || SetAreCurvesVisibleHandler.IsBound() == false)
+	{
+		bAreCurvesVisible = !IsLinearColorCurve();
+	}
 
-	bAreCurvesVisible = !IsLinearColorCurve();
 	bIsGradientEditorVisible = IsLinearColorCurve();
 
 	CurveViewModels.Empty();
@@ -1980,7 +1985,7 @@ void SCurveEditor::ProcessClick(const FGeometry& InMyGeometry, const FPointerEve
 					//If the user is holding shift-ctrl we snap all curve to the mouse position. (false value)
 					//If the user is holding shift we snap to mouse only if there is only one editable curve. (false value)
 					//In all other case we add key directly on the curve. (true value)
-					bAddKeysInline = (IsLinearColorCurve() && !bAreCurvesVisible) || ((!bControlDown) && (CurvesToAddKeysTo->Num() != 1));
+					bAddKeysInline = (IsLinearColorCurve() && !bAreCurvesVisible.Get()) || ((!bControlDown) && (CurvesToAddKeysTo->Num() != 1));
 				}
 				AddNewKey(InMyGeometry, InMouseEvent.GetScreenSpacePosition(), CurvesToAddKeysTo, bAddKeysInline);
 			}
@@ -2774,7 +2779,7 @@ void SCurveEditor::CreateContextMenu(const FGeometry& InMyGeometry, const FPoint
 				MenuBuilder.AddMenuEntry(MenuItemLabel, MenuItemToolTip, FSlateIcon(), Action);
 				
 				//This menu is not required when there is no curve display (color track can hide and show curves)
-				if (bAreCurvesVisible)
+				if (bAreCurvesVisible.Get())
 				{
 					//add key and value to all curve menu entry
 					MenuItemLabel = LOCTEXT("AddKeyValueToAllCurves", "Add key & value to all curves");
@@ -2841,6 +2846,17 @@ void SCurveEditor::OnCreateExternalCurveClicked()
 	OnCreateAsset.ExecuteIfBound();
 }
 
+void SCurveEditor::OnShowCurveToggled()
+{
+	if (bAreCurvesVisible.IsBound() && SetAreCurvesVisibleHandler.IsBound())
+	{
+		SetAreCurvesVisibleHandler.Execute(!bAreCurvesVisible.Get());
+	}
+	else
+	{
+		bAreCurvesVisible = !bAreCurvesVisible.Get();
+	}
+}
 
 UObject* SCurveEditor::CreateCurveObject( TSubclassOf<UCurveBase> CurveType, UObject* PackagePtr, FName& AssetName )
 {
diff --git a/Engine/Source/Editor/UnrealEd/Public/SCurveEditor.h b/Engine/Source/Editor/UnrealEd/Public/SCurveEditor.h
--- a/Engine/Source/Editor/UnrealEd/Public/SCurveEditor.h
+++ b/Engine/Source/Editor/UnrealEd/Public/SCurveEditor.h
@@ -126,6 +126,7 @@ public:
 
 DECLARE_DELEGATE_TwoParams( FOnSetInputViewRange, float, float )
 DECLARE_DELEGATE_TwoParams( FOnSetOutputViewRange, float, float )
+DECLARE_DELEGATE_OneParam( FOnSetAreCurvesVisible, bool )
 
 class SCurveEditor : 
 	public SCompoundWidget,
@@ -208,6 +209,7 @@ public:
 		SLATE_ATTRIBUTE( bool, ShowTimeInFrames )
 		SLATE_ATTRIBUTE( float, TimelineLength )
 		SLATE_ATTRIBUTE( FVector2D, DesiredSize )
+		SLATE_ATTRIBUTE( bool, AreCurvesVisible )
 		SLATE_ARGUMENT( bool, DrawCurve )
 		SLATE_ARGUMENT( bool, HideUI )
 		SLATE_ARGUMENT( bool, AllowZoomOutput )
@@ -223,6 +225,7 @@ public:
 		SLATE_ARGUMENT( FLinearColor, GridColor )
 		SLATE_EVENT( FOnSetInputViewRange, OnSetInputViewRange )
 		SLATE_EVENT( FOnSetOutputViewRange, OnSetOutputViewRange )
+		SLATE_EVENT( FOnSetAreCurvesVisible, OnSetAreCurvesVisible )
 		SLATE_EVENT( FSimpleDelegate, OnCreateAsset )
 	SLATE_END_ARGS()
 
@@ -450,7 +453,7 @@ private:
 	void OnCreateExternalCurveClicked();
 
 	/** Called when "Show Curves" is selected from the context menu */
-	void OnShowCurveToggled() { bAreCurvesVisible = !bAreCurvesVisible; }
+	void OnShowCurveToggled();
 
 	/** Called when "Show Gradient" is selected from the context menu */
 	void OnShowGradientToggled() { bIsGradientEditorVisible = !bIsGradientEditorVisible; }
@@ -573,7 +576,7 @@ private:
 	UNREALED_API virtual void PostRedo(bool bSuccess) override { PostUndo(bSuccess); }
 	// End of FEditorUndoClient
 
-	bool AreCurvesVisible() const { return bAlwaysDisplayColorCurves || bAreCurvesVisible; }
+	bool AreCurvesVisible() const { return bAlwaysDisplayColorCurves || bAreCurvesVisible.Get(); }
 	bool IsGradientEditorVisible() const { return bIsGradientEditorVisible; }
 	bool IsLinearColorCurve() const;
 
@@ -693,6 +696,9 @@ private:
 	/** Handler for adjust timeline panning viewing */
 	FOnSetOutputViewRange	SetOutputViewRangeHandler;
 
+	/** Handler for setting whether or not curves are being displayed. */
+	FOnSetAreCurvesVisible SetAreCurvesVisibleHandler;
+
 	/** Index for the current transaction if any */
 	int32					TransactionIndex;
 
@@ -702,9 +708,6 @@ private:
 
 	/**Flag to enable/disable track editing*/
 	bool		bCanEditTrack;
-
-	/** True if the curves are being displayed */
-	bool bAreCurvesVisible;
 	
 	/** True if the gradient editor is being displayed */
 	bool bIsGradientEditorVisible;
@@ -746,6 +749,9 @@ protected:
 	/** Show time in frames. */
 	TAttribute<bool> bShowTimeInFrames;
 
+	/** Whether or not curves are being displayed */
+	TAttribute<bool> bAreCurvesVisible;
+
 	/** True if you want the curve editor to fit to zoom **/
 	bool bZoomToFitVertical;
 
diff --git a/Engine/Source/Editor/UnrealEd/UnrealEd.Build.cs b/Engine/Source/Editor/UnrealEd/UnrealEd.Build.cs
--- a/Engine/Source/Editor/UnrealEd/UnrealEd.Build.cs
+++ b/Engine/Source/Editor/UnrealEd/UnrealEd.Build.cs
@@ -41,7 +41,6 @@ public class UnrealEd : ModuleRules
 				"MeshUtilities",
 				"Messaging",
 				"MovieSceneCapture",
-				"NiagaraEditor",
 				"PlacementMode",
 				"Settings",
 				"SettingsEditor",
@@ -52,7 +51,6 @@ public class UnrealEd : ModuleRules
 				"IntroTutorials",
 				"OutputLog",
 				"Landscape",
-				"Niagara",
 				"SizeMap",
 				"LocalizationService",
 				"HierarchicalLODUtilities",
@@ -201,7 +199,6 @@ public class UnrealEd : ModuleRules
 				"IntroTutorials",
 				"WorkspaceMenuStructure",
 				"PlacementMode",
-				"NiagaraEditor",
 				"MeshUtilities",
 				"MergeActors",
 				"ProjectSettingsViewer",
diff --git a/Engine/Source/Editor/ViewportInteraction/MouseCursorInteractor.cpp b/Engine/Source/Editor/ViewportInteraction/MouseCursorInteractor.cpp
--- a/Engine/Source/Editor/ViewportInteraction/MouseCursorInteractor.cpp
+++ b/Engine/Source/Editor/ViewportInteraction/MouseCursorInteractor.cpp
@@ -33,7 +33,7 @@ void UMouseCursorInteractor::PollInput()
 	InteractorData.LastRoomSpaceTransform = InteractorData.RoomSpaceTransform;
 
 	// Make sure we have a valid viewport with a cursor over it, and that the viewport's world is the same as ours
-	FEditorViewportClient* ViewportClientPtr = GCurrentLevelEditingViewportClient;
+	FEditorViewportClient* ViewportClientPtr = WorldInteraction->GetDefaultOptionalViewportClient();
 	if( ViewportClientPtr != nullptr && ViewportClientPtr->GetWorld() == WorldInteraction->GetWorld() )
 	{
 		FEditorViewportClient& ViewportClient = *ViewportClientPtr;
diff --git a/Engine/Source/Editor/ViewportInteraction/Public/ViewportInteractableInterface.h b/Engine/Source/Editor/ViewportInteraction/Public/ViewportInteractableInterface.h
--- a/Engine/Source/Editor/ViewportInteraction/Public/ViewportInteractableInterface.h
+++ b/Engine/Source/Editor/ViewportInteraction/Public/ViewportInteractableInterface.h
@@ -43,4 +43,7 @@ public:
 
 	/** Get dragging operation */
 	virtual class UViewportDragOperationComponent* GetDragOperationComponent() PURE_VIRTUAL( IViewportInteractableInterface::GetDragOperationComponent, return nullptr; );
+
+	/** Whether this interactable can be selected. */
+	virtual bool CanBeSelected() PURE_VIRTUAL( IViewportInteractableInterface::CanBeSelected, return false; );
 };
diff --git a/Engine/Source/Editor/ViewportInteraction/Public/ViewportWorldInteraction.h b/Engine/Source/Editor/ViewportInteraction/Public/ViewportWorldInteraction.h
--- a/Engine/Source/Editor/ViewportInteraction/Public/ViewportWorldInteraction.h
+++ b/Engine/Source/Editor/ViewportInteraction/Public/ViewportWorldInteraction.h
@@ -454,6 +454,9 @@ private:
 	/** If there is a transformable with velocity in simulate */
 	bool HasTransformableWithVelocityInSimulate() const;
 
+	/** Get mode tools from the viewport client. */
+	class FEditorModeTools& GetModeTools() const;
+
 	//
 	// Colors
 	//
diff --git a/Engine/Source/Editor/ViewportInteraction/ViewportInteractor.cpp b/Engine/Source/Editor/ViewportInteraction/ViewportInteractor.cpp
--- a/Engine/Source/Editor/ViewportInteraction/ViewportInteractor.cpp
+++ b/Engine/Source/Editor/ViewportInteraction/ViewportInteractor.cpp
@@ -176,18 +176,24 @@ bool UViewportInteractor::HandleInputKey( FEditorViewportClient& ViewportClient,
 							const bool bHaveLaserPointer = GetLaserPointer( /* Out */ LaserPointerStart, /* Out */ LaserPointerEnd );
 							check( bHaveLaserPointer );
 
+							bool bCanBeSelected = true;
 							if ( IViewportInteractableInterface* ActorInteractable = Cast<IViewportInteractableInterface>( Actor ) )
 							{
-								bHandled = true;
-								bool bResultedInInteractableDrag = false;
-								ActorInteractable->OnPressed( this, HitResult, bResultedInInteractableDrag );
-
-								if ( bResultedInInteractableDrag )
+								bCanBeSelected = ActorInteractable->CanBeSelected();
+								if (!bCanBeSelected)
 								{
-									WorldInteraction->SetDraggedInteractable(ActorInteractable, this);
+									bHandled = true;
+									bool bResultedInInteractableDrag = false;
+									ActorInteractable->OnPressed(this, HitResult, bResultedInInteractableDrag);
+
+									if (bResultedInInteractableDrag)
+									{
+										WorldInteraction->SetDraggedInteractable(ActorInteractable, this);
+									}
 								}
 							}
-							else
+
+							if (bCanBeSelected)
 							{
 								bHandled = true;
 
