Copying Orion-Staging @ 3748653 (Orion/Dev-General @ 3722124) to //UE4/Main

diff --git a/Engine/Shaders/Private/PostProcessTonemap.usf b/Engine/Shaders/Private/PostProcessTonemap.usf
--- a/Engine/Shaders/Private/PostProcessTonemap.usf
+++ b/Engine/Shaders/Private/PostProcessTonemap.usf
@@ -119,14 +119,14 @@ float2 UVToScreenPos(float2 UV)
 	return (UV * PostprocessInput0Size.xy - ScreenPosToPixel.zw) / ScreenPosToPixel.xy;
 }
 
-float4 FringeUVParams;
+float4 ChromaticAberrationParams;
 
 void TonemapCommonVS(
 	in	float4 Position,
 	in	float2 TexCoord,
 	out	float3 OutExposureScaleVignette,
 	out	float4 OutGrainUV,
-	out	float4 OutFringe,
+	out	float2 OutScreenPos,
 	out	float2 OutFullViewUV
 	)
 {
@@ -154,9 +154,7 @@ void TonemapCommonVS(
 	OutGrainUV.zw = TexCoord + GrainRandomFull.xy;
 
 	// Fringe
-	float2 ScreenPos = UVToScreenPos(TexCoord);
-	OutFringe.xy = ScreenPosToUV(ScreenPos * FringeUVParams.r);
-	OutFringe.zw = ScreenPosToUV(ScreenPos * FringeUVParams.g);
+	OutScreenPos = UVToScreenPos(TexCoord);
 }
 
 // vertex shader entry point
@@ -166,7 +164,7 @@ void MainVS(
 	out noperspective 	float2 OutTexCoord 				: TEXCOORD0,
 	out noperspective 	float3 OutExposureScaleVignette : TEXCOORD1,
 	out noperspective 	float4 OutGrainUV 				: TEXCOORD2,
-	out noperspective 	float4 OutFringe 				: TEXCOORD3,
+	out noperspective 	float2 OutScreenPos 			: TEXCOORD3,
 	out noperspective 	float2 OutFullViewUV 			: TEXCOORD4,
 	out 				float4 OutPosition 				: SV_POSITION
 	)
@@ -177,7 +175,7 @@ void MainVS(
 #endif
 
 	DrawRectangle(InPosition, InTexCoord, OutPosition, OutTexCoord);
-	TonemapCommonVS(OutPosition, OutTexCoord, OutExposureScaleVignette, OutGrainUV, OutFringe, OutFullViewUV);
+    TonemapCommonVS(OutPosition, OutTexCoord, OutExposureScaleVignette, OutGrainUV, OutScreenPos, OutFullViewUV);
 }
 
 // Function graphing
@@ -290,7 +288,7 @@ float4 TonemapCommonPS(
 	float2 UV,
 	float3 ExposureScaleVignette,
 	float4 GrainUV,
-	float4 FringeUV,
+	float2 ScreenPos, // [-1, 1]x[-1, 1]
 	float2 FullViewUV,
 	float4 SvPosition
 	)
@@ -317,7 +315,7 @@ float4 TonemapCommonPS(
 
 		OffsetInterpolator(CounterStep, UV);
 		OffsetInterpolator(CounterStep, GrainUV);
-		OffsetInterpolator(CounterStep, FringeUV);
+		OffsetInterpolator(CounterStep, ScreenPos);
 		OffsetInterpolator(CounterStep, FullViewUV);
 	}
 	if (all(DebugTile == int2(0, 0)))
@@ -340,8 +338,17 @@ float4 TonemapCommonPS(
 		#if USE_GRAIN_JITTER
 			SceneUVJitter = SceneUV.xy - UV.xy;
 		#endif
-		half3 SceneColor = Texture2DSample(PostprocessInput0, PostprocessInput0Sampler, FringeUV.zw + SceneUVJitter.xy).rgb;
-		half SceneColorG = Texture2DSample(PostprocessInput0, PostprocessInput0Sampler, FringeUV.xy + SceneUVJitter.xy).g;
+		
+        float2 CAScale = ChromaticAberrationParams.xy;
+        float StartOffset = ChromaticAberrationParams.z;
+
+        float4 CAUV;
+        CAUV = ScreenPos.xyxy - sign(ScreenPos).xyxy * saturate(abs(ScreenPos) - StartOffset).xyxy * CAScale.xxyy;
+        CAUV.xy = ScreenPosToUV(CAUV.xy);
+        CAUV.zw = ScreenPosToUV(CAUV.zw);
+
+		half3 SceneColor = Texture2DSample(PostprocessInput0, PostprocessInput0Sampler, CAUV.xy + SceneUVJitter.xy).rgb;
+		half SceneColorG = Texture2DSample(PostprocessInput0, PostprocessInput0Sampler, CAUV.zw + SceneUVJitter.xy).g;
 		half SceneColorB = Texture2DSample(PostprocessInput0, PostprocessInput0Sampler, SceneUV).b;
 		SceneColor.g = SceneColorG;
 		SceneColor.b = SceneColorB;
@@ -536,13 +543,13 @@ void MainPS(
 	in noperspective float2 UV : TEXCOORD0,
 	in noperspective float3 InExposureScaleVignette : TEXCOORD1,
 	in noperspective float4 GrainUV : TEXCOORD2,
-	in noperspective float4 FringeUV : TEXCOORD3,
+	in noperspective float2 ScreenPos : TEXCOORD3,
 	in noperspective float2 FullViewUV : TEXCOORD4,
 	float4 SvPosition : SV_POSITION,		// after all interpolators
 	out float4 OutColor : SV_Target0
 	)
 {
-	OutColor = TonemapCommonPS(UV, InExposureScaleVignette, GrainUV, FringeUV, FullViewUV, SvPosition);
+    OutColor = TonemapCommonPS(UV, InExposureScaleVignette, GrainUV, ScreenPos, FullViewUV, SvPosition);
 }
 
 #if COMPUTESHADER
@@ -562,11 +569,11 @@ void MainCS(uint2 DispatchThreadId : SV_DispatchThreadID)
 	}
 
 	float3 ExposureScaleVignette;
-	float4 GrainUV, FringeUV;
-	float2 FullViewUV;	
-	TonemapCommonVS(SvPosition, UV, ExposureScaleVignette, GrainUV, FringeUV, FullViewUV);
+	float4 GrainUV;
+	float2 FullViewUV, ScreenPos;	
+	TonemapCommonVS(SvPosition, UV, ExposureScaleVignette, GrainUV, ScreenPos, FullViewUV);
 
-	float4 OutColor = TonemapCommonPS(UV, ExposureScaleVignette, GrainUV, FringeUV, FullViewUV, SvPosition);
+	float4 OutColor = TonemapCommonPS(UV, ExposureScaleVignette, GrainUV, ScreenPos, FullViewUV, SvPosition);
 
 	uint2 PixelPos = DispatchThreadId + (uint2)View.ViewRectMin.xy;
 #if NEEDTOSWITCHVERTICLEAXIS
diff --git a/Engine/Shaders/Public/ShaderVersion.ush b/Engine/Shaders/Public/ShaderVersion.ush
--- a/Engine/Shaders/Public/ShaderVersion.ush
+++ b/Engine/Shaders/Public/ShaderVersion.ush
@@ -2,4 +2,4 @@
 // Each time the console command is executed it generates a new GUID. As this file is included
 // in Platform.ush (which should be included in any shader) it allows to invalidate the shader DDC.
 // 
-// GUID = 4C6AEAAADB0F4C96A1C49C75F8B4EB16
\ No newline at end of file
+// GUID = 2B48262180DD484D8B43B113D3708AE1
diff --git a/Engine/Source/Developer/AssetTools/Private/AssetTools.cpp b/Engine/Source/Developer/AssetTools/Private/AssetTools.cpp
--- a/Engine/Source/Developer/AssetTools/Private/AssetTools.cpp
+++ b/Engine/Source/Developer/AssetTools/Private/AssetTools.cpp
@@ -256,6 +256,24 @@ TWeakPtr<IAssetTypeActions> UAssetToolsImpl::GetAssetTypeActionsForClass( UClass
 	return MostDerivedAssetTypeActions;
 }
 
+TArray<TWeakPtr<IAssetTypeActions>> UAssetToolsImpl::GetAssetTypeActionsListForClass(UClass* Class) const
+{
+	TArray<TWeakPtr<IAssetTypeActions>> ResultAssetTypeActionsList;
+
+	for (int32 TypeActionsIdx = 0; TypeActionsIdx < AssetTypeActionsList.Num(); ++TypeActionsIdx)
+	{
+		TSharedRef<IAssetTypeActions> TypeActions = AssetTypeActionsList[TypeActionsIdx];
+		UClass* SupportedClass = TypeActions->GetSupportedClass();
+
+		if (Class->IsChildOf(SupportedClass))
+		{
+			ResultAssetTypeActionsList.Add(TypeActions);
+		}
+	}
+
+	return ResultAssetTypeActionsList;
+}
+
 EAssetTypeCategories::Type UAssetToolsImpl::RegisterAdvancedAssetCategory(FName CategoryKey, FText CategoryDisplayName)
 {
 	EAssetTypeCategories::Type Result = FindAdvancedAssetCategory(CategoryKey);
diff --git a/Engine/Source/Developer/AssetTools/Private/AssetTools.h b/Engine/Source/Developer/AssetTools/Private/AssetTools.h
--- a/Engine/Source/Developer/AssetTools/Private/AssetTools.h
+++ b/Engine/Source/Developer/AssetTools/Private/AssetTools.h
@@ -56,6 +56,7 @@ public:
 	virtual void UnregisterAssetTypeActions(const TSharedRef<IAssetTypeActions>& ActionsToRemove) override;
 	virtual void GetAssetTypeActionsList( TArray<TWeakPtr<IAssetTypeActions>>& OutAssetTypeActionsList ) const override;
 	virtual TWeakPtr<IAssetTypeActions> GetAssetTypeActionsForClass( UClass* Class ) const override;
+	virtual TArray<TWeakPtr<IAssetTypeActions>> GetAssetTypeActionsListForClass(UClass* Class) const override;
 	virtual EAssetTypeCategories::Type RegisterAdvancedAssetCategory(FName CategoryKey, FText CategoryDisplayName) override;
 	virtual EAssetTypeCategories::Type FindAdvancedAssetCategory(FName CategoryKey) const override;
 	virtual void GetAllAdvancedAssetCategories(TArray<FAdvancedAssetCategory>& OutCategoryList) const override;
diff --git a/Engine/Source/Developer/AssetTools/Public/AssetTypeActions_Base.h b/Engine/Source/Developer/AssetTools/Public/AssetTypeActions_Base.h
--- a/Engine/Source/Developer/AssetTools/Public/AssetTypeActions_Base.h
+++ b/Engine/Source/Developer/AssetTools/Public/AssetTypeActions_Base.h
@@ -131,6 +131,11 @@ public:
 		InFilter.ClassNames.Add(GetSupportedClass()->GetFName());
 		InFilter.bRecursiveClasses = true;
 	}
+	
+	virtual FText GetDisplayNameFromAssetData(const FAssetData& AssetData) const
+	{
+		return FText::GetEmpty();
+	}
 
 protected:
 
diff --git a/Engine/Source/Developer/AssetTools/Public/IAssetTools.h b/Engine/Source/Developer/AssetTools/Public/IAssetTools.h
--- a/Engine/Source/Developer/AssetTools/Public/IAssetTools.h
+++ b/Engine/Source/Developer/AssetTools/Public/IAssetTools.h
@@ -112,6 +112,9 @@ public:
 	/** Gets the appropriate AssetTypeActions for the supplied class */
 	virtual TWeakPtr<IAssetTypeActions> GetAssetTypeActionsForClass(UClass* Class) const = 0;
 
+	/** Gets the list of appropriate AssetTypeActions for the supplied class */
+	virtual TArray<TWeakPtr<IAssetTypeActions>> GetAssetTypeActionsListForClass(UClass* Class) const = 0;
+
 	/**
 	* Allocates a Category bit for a user-defined Category, or EAssetTypeCategories::Misc if all available bits are allocated.
 	* Ignores duplicate calls with the same CategoryKey (returns the existing bit but does not change the display name).
diff --git a/Engine/Source/Developer/FunctionalTesting/Private/FuncTestRenderingComponent.cpp b/Engine/Source/Developer/FunctionalTesting/Private/FuncTestRenderingComponent.cpp
--- a/Engine/Source/Developer/FunctionalTesting/Private/FuncTestRenderingComponent.cpp
+++ b/Engine/Source/Developer/FunctionalTesting/Private/FuncTestRenderingComponent.cpp
@@ -11,9 +11,15 @@
 //----------------------------------------------------------------------//
 // FFTestRenderingSceneProxy
 //----------------------------------------------------------------------//
-class FFTestRenderingSceneProxy : public FPrimitiveSceneProxy
+class FFTestRenderingSceneProxy final : public FPrimitiveSceneProxy
 {
 public:
+	SIZE_T GetTypeHash() const override
+	{
+		static size_t UniquePointer;
+		return reinterpret_cast<size_t>(&UniquePointer);
+	}
+
 	/** Initialization constructor. */
 	FFTestRenderingSceneProxy(const UFuncTestRenderingComponent& InComponent)
 		: FPrimitiveSceneProxy(&InComponent)
diff --git a/Engine/Source/Developer/LogVisualizer/Private/VisualLoggerRenderingActor.cpp b/Engine/Source/Developer/LogVisualizer/Private/VisualLoggerRenderingActor.cpp
--- a/Engine/Source/Developer/LogVisualizer/Private/VisualLoggerRenderingActor.cpp
+++ b/Engine/Source/Developer/LogVisualizer/Private/VisualLoggerRenderingActor.cpp
@@ -12,9 +12,15 @@
 #include "VisualLoggerRenderingComponent.h"
 
 class UVisualLoggerRenderingComponent;
-class FVisualLoggerSceneProxy : public FDebugRenderSceneProxy
+class FVisualLoggerSceneProxy final : public FDebugRenderSceneProxy
 {
 public:
+	SIZE_T GetTypeHash() const override
+	{
+		static size_t UniquePointer;
+		return reinterpret_cast<size_t>(&UniquePointer);
+	}
+
 	FVisualLoggerSceneProxy(const UVisualLoggerRenderingComponent* InComponent)
 		: FDebugRenderSceneProxy(InComponent)
 	{
diff --git a/Engine/Source/Developer/MeshMergeUtilities/Private/MeshMergeUtilities.cpp b/Engine/Source/Developer/MeshMergeUtilities/Private/MeshMergeUtilities.cpp
--- a/Engine/Source/Developer/MeshMergeUtilities/Private/MeshMergeUtilities.cpp
+++ b/Engine/Source/Developer/MeshMergeUtilities/Private/MeshMergeUtilities.cpp
@@ -1028,7 +1028,7 @@ void FMeshMergeUtilities::CreateProxyMesh(const TArray<AActor*>& InActors, const
 			Actor->GetComponents<UStaticMeshComponent>(Components);
 
 			// Remove anything non-regular or non-spline static mesh components
-			Components.RemoveAll([](UStaticMeshComponent* Val) { return !(Val->GetClass() == UStaticMeshComponent::StaticClass() || Val->IsA(USplineMeshComponent::StaticClass())); });
+			Components.RemoveAll([](UStaticMeshComponent* Val) { return !((Val->GetClass() == UStaticMeshComponent::StaticClass() && Val->GetStaticMesh()) || Val->IsA(USplineMeshComponent::StaticClass())); });
 
 			ComponentsToMerge.Append(Components);
 		}
diff --git a/Engine/Source/Developer/ShaderFormatVectorVM/Private/VectorVMBackend.cpp b/Engine/Source/Developer/ShaderFormatVectorVM/Private/VectorVMBackend.cpp
--- a/Engine/Source/Developer/ShaderFormatVectorVM/Private/VectorVMBackend.cpp
+++ b/Engine/Source/Developer/ShaderFormatVectorVM/Private/VectorVMBackend.cpp
@@ -21,6 +21,27 @@ PRAGMA_POP
 
 #include "VectorVM.h"
 
+#include "IConsoleManager.h"
+
+#include "Stats.h"
+
+DECLARE_STATS_GROUP(TEXT("VectorVMBackend"), STATGROUP_VectorVMBackend, STATCAT_Advanced);
+
+DECLARE_CYCLE_STAT(TEXT("Generate Main"), STAT_VVMGenerateMain, STATGROUP_VectorVMBackend);
+DECLARE_CYCLE_STAT(TEXT("Generate Code"), STAT_VVMGenerateCode, STATGROUP_VectorVMBackend);
+
+DECLARE_CYCLE_STAT(TEXT("Initial/Misc"), STAT_VVMInitMisc, STATGROUP_VectorVMBackend);
+DECLARE_CYCLE_STAT(TEXT("Matriced To Vectors"), STAT_VVMMatToVec, STATGROUP_VectorVMBackend);
+DECLARE_CYCLE_STAT(TEXT("Branches To Selects"), STAT_VVMBranchesToSelects, STATGROUP_VectorVMBackend);
+DECLARE_CYCLE_STAT(TEXT("To Single Op"), STAT_VVMToSingleOp, STATGROUP_VectorVMBackend);
+DECLARE_CYCLE_STAT(TEXT("Scalarize"), STAT_VVMScalarize, STATGROUP_VectorVMBackend);
+DECLARE_CYCLE_STAT(TEXT("Merge Ops"), STAT_VVMMergeOps, STATGROUP_VectorVMBackend);
+DECLARE_CYCLE_STAT(TEXT("Propagate non-expressions"), STAT_VVMPropagateNonExprs, STATGROUP_VectorVMBackend);
+DECLARE_CYCLE_STAT(TEXT("Cleanup"), STAT_VVMCleanup, STATGROUP_VectorVMBackend);
+DECLARE_CYCLE_STAT(TEXT("Validate"), STAT_VVMValidate, STATGROUP_VectorVMBackend);
+DECLARE_CYCLE_STAT(TEXT("Gen Bytecode"), STAT_VVMGenByteCode, STATGROUP_VectorVMBackend);
+
+
 #if !PLATFORM_WINDOWS
 #define _strdup strdup
 #endif
@@ -29,6 +50,7 @@ DEFINE_LOG_CATEGORY(LogVVMBackend);
 
 bool FVectorVMCodeBackend::GenerateMain(EHlslShaderFrequency Frequency, const char* EntryPoint, exec_list* Instructions, _mesa_glsl_parse_state* ParseState)
 {
+	SCOPE_CYCLE_COUNTER(STAT_VVMGenerateMain);
 	//vm_debug_dump(Instructions, ParseState);
 
 	ir_function_signature* MainSig = NULL;
@@ -57,86 +79,119 @@ bool FVectorVMCodeBackend::GenerateMain(EHlslShaderFrequency Frequency, const ch
 
 char* FVectorVMCodeBackend::GenerateCode(exec_list* ir, _mesa_glsl_parse_state* state, EHlslShaderFrequency Frequency)
 {
+	SCOPE_CYCLE_COUNTER(STAT_VVMGenerateCode); 
 	vm_debug_print("========VECTOR VM BACKEND: Generate Code==============\n");
 	vm_debug_dump(ir, state);
 
 	if (state->error) return nullptr;
 
-	//Inline all functions.
-	vm_debug_print("== Initial misc ==\n");
 	bool progress = false;
-	do
 	{
-		//progress = do_function_inlining(ir);//Full optimization pass earlier will have already done this.
-		progress = do_mat_op_to_vec(ir);
-		progress = do_vec_op_to_scalar(ir) || progress;
-		progress = do_vec_index_to_swizzle(ir) || progress;
-		progress = do_copy_propagation(ir) || progress;
-		progress = do_copy_propagation_elements(ir) || progress;
-		progress = do_swizzle_swizzle(ir) || progress;
-	} while (progress);
-	//validate_ir_tree(ir, state);
-	vm_debug_dump(ir, state);
-	if (state->error) return nullptr;
+		SCOPE_CYCLE_COUNTER(STAT_VVMInitMisc);
+		vm_debug_print("== Initial misc ==\n");
+		do
+		{
+			//progress = do_function_inlining(ir);//Full optimization pass earlier will have already done this.
+			progress = do_mat_op_to_vec(ir);
+			progress = do_vec_op_to_scalar(ir, state) || progress;
+			progress = do_vec_index_to_swizzle(ir) || progress;
+			progress = do_copy_propagation(ir) || progress;
+			progress = do_copy_propagation_elements(ir) || progress;
+			progress = do_swizzle_swizzle(ir) || progress;
+		} while (progress);
+		//validate_ir_tree(ir, state);
+		vm_debug_dump(ir, state);
+		if (state->error) return nullptr;
+	}
 	
-	vm_debug_print("== Branches to selects ==\n");
-	vm_flatten_branches_to_selects(ir, state);
-	//validate_ir_tree(ir, state);
-	vm_debug_dump(ir, state);
+	{
+		SCOPE_CYCLE_COUNTER(STAT_VVMMatToVec);
+		//99% complete code to remove all matrices from the code and replace them with just swizzled vectors. 
+		//For now visitors below here can handle matrices ok but we may hit some edge cases in future requiring their removal.
+		vm_debug_print("== matrices to vectors ==\n");
+		vm_matrices_to_vectors(ir, state);
+		vm_debug_dump(ir, state);
+		if (state->error) return nullptr;
+	}
 
-	vm_debug_print("== To Single Op ==\n");
-	vm_to_single_op(ir, state);
-	//validate_ir_tree(ir, state);
-	vm_debug_dump(ir, state);
-	if (state->error) return nullptr;
+	{
+		SCOPE_CYCLE_COUNTER(STAT_VVMBranchesToSelects);
+		vm_debug_print("== Branches to selects ==\n");
+		vm_flatten_branches_to_selects(ir, state);
+		//validate_ir_tree(ir, state);
+		vm_debug_dump(ir, state);
+	}
 
-	vm_debug_print("== Scalarize ==\n");
-	vm_scalarize_ops(ir, state);
-	//validate_ir_tree(ir, state);
-	vm_debug_dump(ir, state);
-	//validate_ir_tree(ir, state);
-	if (state->error) return nullptr;
+	{
 
-	//99% complete code to remove all matrices from the code and replace them with just swizzled vectors. 
-	//For now visitors below here can handle matrices ok but we may hit some edge cases in future requiring their removal.
-	//vm_debug_print("== matrices to vectors ==\n");
-	//vm_matrices_to_vectors(ir, state);
-	//vm_debug_dump(ir, state);
-	//if (state->error) return nullptr;
+		SCOPE_CYCLE_COUNTER(STAT_VVMToSingleOp);
+		vm_debug_print("== To Single Op ==\n");
+		vm_to_single_op(ir, state);
+		//validate_ir_tree(ir, state);
+		vm_debug_dump(ir, state);
+		if (state->error) return nullptr;
+	}
 
-	vm_debug_print("== Merge Ops ==\n");
-	vm_merge_ops(ir, state);
-//	validate_ir_tree(ir, state);
-	vm_debug_dump(ir, state);
-	//validate_ir_tree(ir, state);
-	if (state->error) return nullptr;
+	{
 
-	vm_debug_print("== Propagate non-expressions ==\n");
-	vm_propagate_non_expressions_visitor(ir, state);
-	//validate_ir_tree(ir, state);
-	vm_debug_dump(ir, state);
-	if (state->error) return nullptr;
+		SCOPE_CYCLE_COUNTER(STAT_VVMScalarize);
+		vm_debug_print("== Scalarize ==\n");
+		vm_scalarize_ops(ir, state);
+		//validate_ir_tree(ir, state);
+		vm_debug_dump(ir, state);
+		//validate_ir_tree(ir, state);
+		if (state->error) return nullptr;
+	}
 
-	vm_debug_print("== Cleanup ==\n");
-	// Final cleanup
-	progress = false;
-	do
 	{
-		progress = do_dead_code(ir, false);
-		progress = do_dead_code_local(ir) || progress;
-		progress = do_swizzle_swizzle(ir) || progress;
-		progress = do_noop_swizzle(ir) || progress;
-		progress = do_copy_propagation(ir) || progress;
-		progress = do_copy_propagation_elements(ir) || progress;
-		progress = do_constant_propagation(ir) || progress;
-	} while (progress);
-	vm_debug_dump(ir, state);
+		SCOPE_CYCLE_COUNTER(STAT_VVMMergeOps);
+		vm_debug_print("== Merge Ops ==\n");
+		vm_merge_ops(ir, state);
+	//	validate_ir_tree(ir, state);
+		vm_debug_dump(ir, state);
+		//validate_ir_tree(ir, state);
+		if (state->error) return nullptr;
+	}
 
-	//validate_ir_tree(ir, state);
+	{
+
+		SCOPE_CYCLE_COUNTER(STAT_VVMPropagateNonExprs);
+		vm_debug_print("== Propagate non-expressions ==\n");
+		vm_propagate_non_expressions_visitor(ir, state);
+		//validate_ir_tree(ir, state);
+		vm_debug_dump(ir, state);
+		if (state->error) return nullptr;
+	}
+
+	{
+		SCOPE_CYCLE_COUNTER(STAT_VVMCleanup);
+		vm_debug_print("== Cleanup ==\n");
+		// Final cleanup
+		progress = false;
+		do
+		{
+			progress = do_dead_code(ir, false);
+			progress = do_dead_code_local(ir) || progress;
+			progress = do_swizzle_swizzle(ir) || progress;
+			progress = do_noop_swizzle(ir) || progress;
+			progress = do_copy_propagation(ir) || progress;
+			progress = do_copy_propagation_elements(ir) || progress;
+			progress = do_constant_propagation(ir) || progress;
+		} while (progress);
+		vm_debug_dump(ir, state);
+	}
+
+	{
+		SCOPE_CYCLE_COUNTER(STAT_VVMValidate);
+		validate_ir_tree(ir, state);
+	}
 
 	if (state->error) return nullptr;
+	{
+		SCOPE_CYCLE_COUNTER(STAT_VVMGenByteCode);
+		vm_gen_bytecode(ir, state, CompilationOutput);
+	}
 
-	vm_gen_bytecode(ir, state, CompilationOutput);
 	return  nullptr;// Cheat and emit the bytecode into he CompilationOutput. The return here is treat as a string so the 0's it contains prove problematic.
 }
 
@@ -157,3 +212,19 @@ void FVectorVMLanguageSpec::SetupLanguageIntrinsics(_mesa_glsl_parse_state* Stat
 // 	make_intrinsic_genType(ir, State, "OutputDataInt", ir_invalid_opcode, 0, 3, 1, 1);
 //	make_intrinsic_genType(ir, State, "AcquireIndex", ir_invalid_opcode, IR_INTRINSIC_INT, 2, 1, 1);
 }
+
+static int32 GbDumpVMIR = 0;
+static FAutoConsoleVariableRef CVarDumpVMIR(
+	TEXT("fx.DumpVMIR"),
+	GbDumpVMIR,
+	TEXT("If > 0 verbose logging is enabled for the vm compiler backend. \n"),
+	ECVF_Default
+);
+
+void DebugDumpIR(struct exec_list* ir, struct _mesa_glsl_parse_state* State)
+{
+	if (GbDumpVMIR)
+	{
+		IRDump(ir, State);
+	}
+}
\ No newline at end of file
diff --git a/Engine/Source/Developer/ShaderFormatVectorVM/Private/VectorVMBackend.h b/Engine/Source/Developer/ShaderFormatVectorVM/Private/VectorVMBackend.h
--- a/Engine/Source/Developer/ShaderFormatVectorVM/Private/VectorVMBackend.h
+++ b/Engine/Source/Developer/ShaderFormatVectorVM/Private/VectorVMBackend.h
@@ -71,7 +71,7 @@ ECallScalarizeMode get_scalarize_mode(ir_function_signature* in_sig);
 EVectorVMOp get_special_vm_opcode(ir_function_signature* signature);
 
 void vm_matrices_to_vectors(exec_list* instructions, _mesa_glsl_parse_state *state);
-bool do_vec_op_to_scalar(exec_list *instructions);
+bool do_vec_op_to_scalar(exec_list *instructions, _mesa_glsl_parse_state* parse_state);
 bool vm_flatten_branches_to_selects(exec_list *instructions, _mesa_glsl_parse_state *state);
 void vm_to_single_op(exec_list *ir, _mesa_glsl_parse_state *state);
 void vm_merge_ops(exec_list *ir, _mesa_glsl_parse_state *state);
@@ -80,11 +80,13 @@ void vm_propagate_non_expressions_visitor(exec_list* ir, _mesa_glsl_parse_state*
 void vm_gen_bytecode(exec_list *ir, _mesa_glsl_parse_state *state, FVectorVMCompilationOutput& InCompOutput);
 
 
+void DebugDumpIR(struct exec_list* ir, struct _mesa_glsl_parse_state* State);
+
 //////////////////////////////////////////////////////////////////////////
 //Enable verbose debug dumps.
 #define VM_VERBOSE_LOGGING 0
 #if VM_VERBOSE_LOGGING == 2
-#define vm_debug_dump(ir, state) IRDump(ir, state)
+#define vm_debug_dump(ir, state) DebugDumpIR(ir, state)
 #define vm_debug_print dprintf
 #elif VM_VERBOSE_LOGGING == 1
 #define vm_debug_dump(ir, state) 
diff --git a/Engine/Source/Developer/ShaderFormatVectorVM/Private/VectorVMShaderCompiler.cpp b/Engine/Source/Developer/ShaderFormatVectorVM/Private/VectorVMShaderCompiler.cpp
--- a/Engine/Source/Developer/ShaderFormatVectorVM/Private/VectorVMShaderCompiler.cpp
+++ b/Engine/Source/Developer/ShaderFormatVectorVM/Private/VectorVMShaderCompiler.cpp
@@ -14,6 +14,12 @@
 
 DEFINE_LOG_CATEGORY_STATIC(LogVectorVMShaderCompiler, Log, All); 
 
+DECLARE_STATS_GROUP(TEXT("VectorVM"), STATGROUP_VectorVM, STATCAT_Advanced);
+
+DECLARE_CYCLE_STAT(TEXT("VectorVM - Compiler - CompileShader_VectorVM"), STAT_VectorVM_Compiler_CompileShader_VectorVM, STATGROUP_VectorVM);
+DECLARE_CYCLE_STAT(TEXT("VectorVM - Compiler - PreprocessShader"), STAT_VectorVM_Compiler_CompileShader_VectorVMPreprocessShader, STATGROUP_VectorVM);
+DECLARE_CYCLE_STAT(TEXT("VectorVM - Compiler - CrossCompilerContextRun"), STAT_VectorVM_Compiler_CompileShader_CrossCompilerContextRun, STATGROUP_VectorVM);
+
 
 /**
  * Compile a shader for the VectorVM on Windows.
@@ -28,6 +34,8 @@ bool CompileShader_VectorVM(const FShaderCompilerInput& Input, FShaderCompilerOu
 //TODO: Move to this output living in the shader eco-system with the compute shaders too but for now just do things more directly.
 bool CompileShader_VectorVM(const FShaderCompilerInput& Input, FShaderCompilerOutput& Output,const FString& WorkingDirectory, uint8 Version, FVectorVMCompilationOutput& VMCompilationOutput)
 {
+	SCOPE_CYCLE_COUNTER(STAT_VectorVM_Compiler_CompileShader_VectorVM);
+
 	FString PreprocessedShader;
 	FShaderCompilerDefinitions AdditionalDefines;
 	EHlslCompileTarget HlslCompilerTarget = HCT_FeatureLevelSM5;
@@ -52,6 +60,8 @@ bool CompileShader_VectorVM(const FShaderCompilerInput& Input, FShaderCompilerOu
 	}
 	else
 	{
+		SCOPE_CYCLE_COUNTER(STAT_VectorVM_Compiler_CompileShader_VectorVMPreprocessShader);
+
 		if (!PreprocessShader(PreprocessedShader, Output, Input, AdditionalDefines))
 		{
 			// The preprocessing stage will add any relevant errors.
@@ -112,12 +122,14 @@ bool CompileShader_VectorVM(const FShaderCompilerInput& Input, FShaderCompilerOu
 	//CCFlags &= ~HLSLCC_NoPreprocess;
 
 	FVectorVMCodeBackend VVMBackEnd(CCFlags, HlslCompilerTarget, VMCompilationOutput);
-	FVectorVMLanguageSpec VVMLanguageSpec; 
+		FVectorVMLanguageSpec VVMLanguageSpec; 
 
 	bool Result = false;
 	FHlslCrossCompilerContext CrossCompilerContext(CCFlags, Frequency, HlslCompilerTarget);
 	if (CrossCompilerContext.Init(TCHAR_TO_ANSI(*Input.VirtualSourceFilePath), &VVMLanguageSpec))
 	{
+		SCOPE_CYCLE_COUNTER(STAT_VectorVM_Compiler_CompileShader_CrossCompilerContextRun);
+
 		Result = CrossCompilerContext.Run(
 			TCHAR_TO_ANSI(*PreprocessedShader),
 			TCHAR_TO_ANSI(*Input.EntryPointName),
@@ -127,7 +139,31 @@ bool CompileShader_VectorVM(const FShaderCompilerInput& Input, FShaderCompilerOu
 			) ? 1 : 0;
 	}
 
-	VMCompilationOutput.Errors = ErrorLog;
+		if (ErrorLog)
+		{
+			int32 SrcLen = FPlatformString::Strlen(ErrorLog);
+			int32 DestLen = FPlatformString::ConvertedLength<TCHAR, ANSICHAR>(ErrorLog, SrcLen);
+			TArray<TCHAR> Converted;
+			Converted.AddUninitialized(DestLen);
+
+			FPlatformString::Convert<ANSICHAR, TCHAR>(Converted.GetData(), DestLen, ErrorLog, SrcLen);
+			Converted.Add(0);
+
+			VMCompilationOutput.Errors = Converted.GetData();
+
+			TArray<FString> OutputByLines;
+			PreprocessedShader.ParseIntoArrayLines(OutputByLines, false);
+
+			FString OutputHlsl;
+			for (int32 i = 0; i < OutputByLines.Num(); i++)
+			{
+				UE_LOG(LogVectorVMShaderCompiler, Warning, TEXT("/*%d*/%s"), i, *OutputByLines[i]);
+			}
+		}
+		else
+		{
+			VMCompilationOutput.Errors.Empty();
+		}
 
 	//TODO: Try to get rid of the CompilationOutput and have the vm bytecode life in the shader eco-system as the compute shader version will.
 //		int32 SourceLen = VVMBackEnd.ByteCode.Num();//ShaderSource ? FCStringAnsi::Strlen(ShaderSource) : 0;
@@ -202,9 +238,9 @@ bool CompileShader_VectorVM(const FShaderCompilerInput& Input, FShaderCompilerOu
 		UE_LOG(LogVectorVMShaderCompiler, Warning, TEXT("%s"), (const char*)ShaderSource);
 		free(ShaderSource);
 	}
-	if (ErrorLog)
+		if (VMCompilationOutput.Errors.Len() != 0)
 	{
-		UE_LOG(LogVectorVMShaderCompiler, Warning, TEXT("%s"), (const char*)ErrorLog);
+			UE_LOG(LogVectorVMShaderCompiler, Warning, TEXT("%s"), *VMCompilationOutput.Errors);
 		free(ErrorLog);
 	}
 
diff --git a/Engine/Source/Developer/ShaderFormatVectorVM/Private/ir_vm_flatten_branches_to_selects_visitor.cpp b/Engine/Source/Developer/ShaderFormatVectorVM/Private/ir_vm_flatten_branches_to_selects_visitor.cpp
--- a/Engine/Source/Developer/ShaderFormatVectorVM/Private/ir_vm_flatten_branches_to_selects_visitor.cpp
+++ b/Engine/Source/Developer/ShaderFormatVectorVM/Private/ir_vm_flatten_branches_to_selects_visitor.cpp
@@ -108,15 +108,15 @@ public:
 
 		ir_if* old_if = curr_if;
 		curr_if = ir;
-		void* parent = ralloc_parent(curr_if);
+		//void* parent = ralloc_parent(curr_if);
 
 		//Pull out our condition to a variable.
-		ir_variable *condition_var = new(parent)ir_variable(glsl_type::bool_type,
+		ir_variable *condition_var = new(parse_state)ir_variable(glsl_type::bool_type,
 			"branch_flatten_condition",
 			ir_var_temporary);
 		base_ir->insert_before(condition_var);
-		base_ir->insert_before(new(parent)ir_assignment(new(parent)ir_dereference_variable(condition_var), curr_if->condition));
-		curr_if->condition = new(parent)ir_dereference_variable(condition_var);
+		base_ir->insert_before(new(parse_state)ir_assignment(new(parse_state)ir_dereference_variable(condition_var), curr_if->condition));
+		curr_if->condition = new(parse_state)ir_dereference_variable(condition_var);
 
 		check(a_assignments.Num() == 0);
 		check(curr_assignments == nullptr);
@@ -129,10 +129,11 @@ public:
 		visit_list_elements(this, &curr_if->else_instructions, true);
 		base_ir->insert_before(&curr_if->else_instructions);
 
-		while (a_assignments.Num() > 0)
-		{
-			ir_assignment* assign = a_assignments.Pop(false);
+		TArray<bool> HandledBAssignments;
+		HandledBAssignments.SetNumZeroed(b_assignments.Num());
 
+		for(ir_assignment* assign : a_assignments)
+		{
 			ir_rvalue* selection_other = nullptr;
 			int32 selection_other_idx = b_assignments.IndexOfByPredicate([&](ir_assignment* other_assign) { return AreEquivalent(assign->lhs, other_assign->lhs); });
 			if (selection_other_idx == INDEX_NONE)
@@ -145,52 +146,59 @@ public:
 			{
 				//We did find a matching statement, select using that and remove it from the b_assignments.
 				selection_other = b_assignments[selection_other_idx]->rhs;
-				b_assignments.RemoveAt(selection_other_idx, 1, false);
+				HandledBAssignments[selection_other_idx] = true;
 			}
 
 			//Selection result
-			ir_variable *result = new(parent)ir_variable(assign->lhs->type, "selction_result", ir_var_temporary);
+			ir_variable *result = new(parse_state)ir_variable(assign->lhs->type, "selction_result", ir_var_temporary);
 			base_ir->insert_before(result);
 
 			exec_list select_params;
-			select_params.push_tail(ir->condition->clone(parent, nullptr));
-			select_params.push_tail(assign->rhs->clone(parent, nullptr));
-			select_params.push_tail(selection_other->clone(parent, nullptr));
+			select_params.push_tail(ir->condition->clone(parse_state, nullptr));
+			select_params.push_tail(assign->rhs->clone(parse_state, nullptr));
+			select_params.push_tail(selection_other->clone(parse_state, nullptr));
 
 			ir_function_signature* selection_sig = get_select_signature(assign->lhs->type);
 			if (selection_sig)
 			{
-				ir_call* select_call = new(parent)ir_call(selection_sig, new(parent)ir_dereference_variable(result), &select_params);
+				ir_call* select_call = new(parse_state)ir_call(selection_sig, new(parse_state)ir_dereference_variable(result), &select_params);
 				base_ir->insert_before(select_call);
 			}
-			assign->rhs = new(parent)ir_dereference_variable(result);
+			assign->rhs = new(parse_state)ir_dereference_variable(result);
 			base_ir->insert_before(assign);
 		}
 
 		//Deal with any remaining work done on the b path.
-		while (b_assignments.Num() > 0)
+		for(int32 i=0; i < b_assignments.Num(); ++i)
 		{
-			ir_assignment* assign = b_assignments.Pop(false);
+			//bail if this assignment was already handled.
+			if(HandledBAssignments[i])
+				continue;
+
+			ir_assignment* assign = b_assignments[i];
 
 			//Selection result
-			ir_variable *result = new(parent)ir_variable(assign->lhs->type, "selction_result", ir_var_temporary);
+			ir_variable *result = new(parse_state)ir_variable(assign->lhs->type, "selction_result", ir_var_temporary);
 			base_ir->insert_before(result);
 
 			exec_list select_params;
-			select_params.push_tail(ir->condition->clone(parent, nullptr));
-			select_params.push_tail(assign->lhs->clone(parent, nullptr));
-			select_params.push_tail(assign->rhs->clone(parent, nullptr));
+			select_params.push_tail(ir->condition->clone(parse_state, nullptr));
+			select_params.push_tail(assign->lhs->clone(parse_state, nullptr));
+			select_params.push_tail(assign->rhs->clone(parse_state, nullptr));
 
 			ir_function_signature* selection_sig = get_select_signature(assign->lhs->type);
 			if (selection_sig)
 			{
-				ir_call* select_call = new(parent)ir_call(selection_sig, new(parent)ir_dereference_variable(result), &select_params);
+				ir_call* select_call = new(parse_state)ir_call(selection_sig, new(parse_state)ir_dereference_variable(result), &select_params);
 				base_ir->insert_before(select_call);
 			}
-			assign->rhs = new(parent)ir_dereference_variable(result);
+			assign->rhs = new(parse_state)ir_dereference_variable(result);
 			base_ir->insert_before(assign);
 		}
 
+		a_assignments.Reset();
+		b_assignments.Reset();
+
 		base_ir->remove();
 		curr_assignments = nullptr;
 		curr_if = old_if;
@@ -201,13 +209,13 @@ public:
 	ir_dereference_variable* replace_assigned_val_with_temp(ir_dereference* val)
 	{
 		check(val);
-		void* parent = ralloc_parent(val);
-		ir_variable *var = new(parent)ir_variable(val->type,
+		//void* parent = ralloc_parent(val);
+		ir_variable *var = new(parse_state)ir_variable(val->type,
 			"branch_flatten_temp",
 			ir_var_temporary);
 		base_ir->insert_before(var);
 
-		return new(parent)ir_dereference_variable(var);
+		return new(parse_state)ir_dereference_variable(var);
 	}
 
 	ir_visitor_status visit_leave(ir_assignment * assign)
@@ -216,10 +224,10 @@ public:
 		{
 			check(curr_assignments);
 
-			void* parent = ralloc_parent(assign);
+			//void* parent = ralloc_parent(assign);
 			ir_dereference_variable* new_deref = replace_assigned_val_with_temp(assign->lhs);
 
-			ir_assignment* new_assign = new(parent)ir_assignment(assign->lhs, new_deref);
+			ir_assignment* new_assign = new(parse_state)ir_assignment(assign->lhs, new_deref);
 			assign->set_lhs(new_deref);
 			curr_assignments->Add(new_assign);
 		}
@@ -233,10 +241,10 @@ public:
 		{
 			check(curr_assignments);
 
-			void* parent = ralloc_parent(call);
+			//void* parent = ralloc_parent(call);
 			ir_dereference_variable* new_deref = replace_assigned_val_with_temp(call->return_deref);
 
-			ir_assignment* new_assign = new(parent)ir_assignment(call->return_deref, new_deref);
+			ir_assignment* new_assign = new(parse_state)ir_assignment(call->return_deref, new_deref);
 			call->return_deref = new_deref;
 			curr_assignments->Add(new_assign);
 
@@ -247,7 +255,7 @@ public:
 				if (var->mode == ir_var_out || var->mode == ir_var_inout)
 				{
 					new_deref = replace_assigned_val_with_temp(actual_param->as_dereference());
-					new_assign = new(parent)ir_assignment(actual_param, new_deref);
+					new_assign = new(parse_state)ir_assignment(actual_param, new_deref);
 					check(actual_param->next && actual_param->prev);
 					actual_param->replace_with(new_deref);
 					curr_assignments->Add(new_assign);
diff --git a/Engine/Source/Developer/ShaderFormatVectorVM/Private/ir_vm_gen_bytecode_visitor.cpp b/Engine/Source/Developer/ShaderFormatVectorVM/Private/ir_vm_gen_bytecode_visitor.cpp
--- a/Engine/Source/Developer/ShaderFormatVectorVM/Private/ir_vm_gen_bytecode_visitor.cpp
+++ b/Engine/Source/Developer/ShaderFormatVectorVM/Private/ir_vm_gen_bytecode_visitor.cpp
@@ -549,7 +549,7 @@ struct op_base
 
 	virtual FString to_string() = 0;
 	virtual void add_to_bytecode(TArray<uint8>& bytecode) = 0;
-	virtual bool finalize_temporary_allocations(const unsigned num_temp_registers, unsigned *registers, unsigned op_idx) { return true; }
+	virtual bool finalize(_mesa_glsl_parse_state* parse_state, FVectorVMCompilationOutput&CompilationOutput, const unsigned num_temp_registers, unsigned *registers, unsigned op_idx) { return true; }
 	virtual void validate(_mesa_glsl_parse_state* parse_state) = 0;
 
 	virtual struct op_standard* as_standard() { return nullptr; }
@@ -708,9 +708,14 @@ struct op_standard : public op_base
 		emit_short(dest_component->offset, bytecode);
 	}
 
-	virtual bool finalize_temporary_allocations(const unsigned num_temp_registers, unsigned *registers, unsigned op_idx)
+	virtual bool finalize(_mesa_glsl_parse_state* parse_state, FVectorVMCompilationOutput&CompilationOutput, const unsigned num_temp_registers, unsigned *registers, unsigned op_idx)
 	{
-		return finalize_component_temporary_allocation(dest_component, num_temp_registers, registers, op_idx);
+		if (!finalize_component_temporary_allocation(dest_component, num_temp_registers, registers, op_idx))
+		{
+			_mesa_glsl_error(parse_state, "Failed to allocate temporary registers");
+			return false;
+		}
+		return true;
 	}
 
 	virtual void validate(_mesa_glsl_parse_state* parse_state)override
@@ -771,9 +776,14 @@ struct op_input : public op_base
 		emit_short(dest_component->offset, bytecode);
 	}
 
-	virtual bool finalize_temporary_allocations(const unsigned num_temp_registers, unsigned *registers, unsigned op_idx)
+	virtual bool finalize(_mesa_glsl_parse_state* parse_state, FVectorVMCompilationOutput&CompilationOutput, const unsigned num_temp_registers, unsigned *registers, unsigned op_idx)
 	{
-		return finalize_component_temporary_allocation(dest_component, num_temp_registers, registers, op_idx);
+		if (!finalize_component_temporary_allocation(dest_component, num_temp_registers, registers, op_idx))
+		{
+			_mesa_glsl_error(parse_state, "Failed to allocate temporary registers");
+			return false;
+		}
+		return true;
 	}
 
 	virtual void validate(_mesa_glsl_parse_state* parse_state)override
@@ -821,9 +831,14 @@ struct op_input_noadvance : public op_input
 		emit_short(dest_component->offset, bytecode);
 	}
 
-	virtual bool finalize_temporary_allocations(const unsigned num_temp_registers, unsigned *registers, unsigned op_idx)
+	virtual bool finalize(_mesa_glsl_parse_state* parse_state, FVectorVMCompilationOutput&CompilationOutput, const unsigned num_temp_registers, unsigned *registers, unsigned op_idx)
 	{
-		return finalize_component_temporary_allocation(dest_component, num_temp_registers, registers, op_idx);
+		if (!finalize_component_temporary_allocation(dest_component, num_temp_registers, registers, op_idx))
+		{
+			_mesa_glsl_error(parse_state, "Failed to allocate temporary registers");
+			return false;
+		}
+		return true;
 	}
 
 	virtual void validate(_mesa_glsl_parse_state* parse_state)override
@@ -920,9 +935,14 @@ struct op_index_acquire : public op_base
 		emit_short(dest_component->offset, bytecode);
 	}
 
-	virtual bool finalize_temporary_allocations(const unsigned num_temp_registers, unsigned *registers, unsigned op_idx)
+	virtual bool finalize(_mesa_glsl_parse_state* parse_state, FVectorVMCompilationOutput&CompilationOutput, const unsigned num_temp_registers, unsigned *registers, unsigned op_idx)
 	{
-		return finalize_component_temporary_allocation(dest_component, num_temp_registers, registers, op_idx);
+		if (!finalize_component_temporary_allocation(dest_component, num_temp_registers, registers, op_idx))
+		{
+			_mesa_glsl_error(parse_state, "Failed to allocate temporary registers");
+			return false;
+		}
+		return true;
 	}
 
 	virtual void validate(_mesa_glsl_parse_state* parse_state)override
@@ -984,8 +1004,29 @@ struct op_external_func : public op_base
 		return Str += TEXT(");\n");
 	}
 
+	bool should_cull_function(unsigned int op_idx)
+	{
+		//If the outputs of this call are never used after the current op then we cull this function.
+		bool func_is_used = false;
+		for (variable_info_node* output : outputs)
+		{
+			if (output->last_read > op_idx)
+			{
+				func_is_used = true;
+				break;
+			}
+		}
+
+		return !func_is_used;
+	}
+
 	virtual void add_to_bytecode(TArray<uint8>& bytecode)
 	{
+		if (function_index == -1)
+		{
+			return;
+		}
+
 		emit_byte((uint8)EVectorVMOp::external_func_call, bytecode);
 
 		emit_byte(function_index, bytecode);
@@ -1016,8 +1057,26 @@ struct op_external_func : public op_base
 		ctx->success &= finalize_component_temporary_allocation(node, ctx->num_temp_registers, ctx->registers, ctx->op_idx);
 	}
 
-	virtual bool finalize_temporary_allocations(const unsigned num_temp_registers, unsigned *registers, unsigned op_idx)
+	virtual bool finalize(_mesa_glsl_parse_state* parse_state, FVectorVMCompilationOutput&CompilationOutput, const unsigned num_temp_registers, unsigned *registers, unsigned op_idx)
 	{
+		//don't allocate any temp registers if we're culling the function.
+		function_index = -1;
+		if (should_cull_function(op_idx))
+		{
+			return false;
+		}
+
+		//We also add this call to the compiler output here as we're now sure it's been called.
+		function_index = CompilationOutput.CalledVMFunctionTable.AddDefaulted();
+		check(function_index != INDEX_NONE);
+		FVectorVMCompilationOutput::FCalledVMFunction& CalledFunc = CompilationOutput.CalledVMFunctionTable[function_index];
+		CalledFunc.Name = sig->function_name();
+		CalledFunc.NumOutputs = outputs.Num();
+		for (variable_info_node* input : inputs)
+		{
+			CalledFunc.InputParamLocations.Add(input->owner->location == EVectorVMOperandLocation::Constant);
+		}
+
 		finalize_temp_register_ctx ctx;
 		ctx.num_temp_registers = num_temp_registers;
 		ctx.registers = registers;
@@ -1027,11 +1086,20 @@ struct op_external_func : public op_base
 		{
 			for_each_component(output, &finalize_comp, true, &ctx);
 		}
+		if (!ctx.success)
+		{
+			_mesa_glsl_error(parse_state, "Failed to allocate temporary registers for external funciton.");
+		}
 		return ctx.success;
 	}
 
 	virtual void validate(_mesa_glsl_parse_state* parse_state)override
 	{
+		if (function_index == -1)
+		{
+			return;
+		}
+
 		for (variable_info_node* input : inputs)
 		{
 			if (input->offset == INDEX_NONE)
@@ -1683,11 +1751,7 @@ class ir_gen_vvm_visitor : public ir_hierarchical_visitor
 			unsigned num_inputs = 0;
 			unsigned num_outputs = 0;
 
-			func->function_index = CompilationOutput.CalledVMFunctionTable.AddDefaulted();
 			func->sig = call->callee;
-			check(func->function_index != INDEX_NONE);
-			FVectorVMCompilationOutput::FCalledVMFunction& CalledFunc = CompilationOutput.CalledVMFunctionTable[func->function_index];
-			CalledFunc.Name = call->callee->function_name();
 
 			exec_node* sig_param_node = call->callee->parameters.get_head();
 			foreach_iter(exec_list_iterator, iter, call->actual_parameters)
@@ -1703,11 +1767,9 @@ class ir_gen_vvm_visitor : public ir_hierarchical_visitor
 				if (sig_param->mode == ir_var_in)
 				{
 					func->inputs.Add(curr_node);
-					CalledFunc.InputParamLocations.Add(curr_node->owner->location == EVectorVMOperandLocation::Constant);
 				}
-				else 
+				else
 				{
-					++CalledFunc.NumOutputs;
 					func->outputs.Add(curr_node);
 				}	
 				curr_node = nullptr;
@@ -1779,15 +1841,20 @@ class ir_gen_vvm_visitor : public ir_hierarchical_visitor
 #if VM_VERBOSE_LOGGING
 		UE_LOG(LogVVMBackend, Log, TEXT("\n-------------------------------\nTemporary Allocations\n------------------------------\n"));
 #endif
+		TArray<int32> ops_to_strip;
 		for (int32 i = 0; i < ordered_ops.Num(); ++i)
 		{
-			if (!ordered_ops[i]->finalize_temporary_allocations(num_temp_registers, registers, i))
+			if (!ordered_ops[i]->finalize(parse_state, CompilationOutput, num_temp_registers, registers, i))
 			{
-				_mesa_glsl_error(parse_state, "register assignment failed");
-				break;
+				ops_to_strip.Add(i);
 			}
 		}
 
+		for (int32 op_idx : ops_to_strip)
+		{
+			ordered_ops.RemoveAt(op_idx);
+		}
+
 		//Final error checking.
 		for (int32 i = 0; i < ordered_ops.Num(); ++i)
 		{
@@ -1806,24 +1873,23 @@ class ir_gen_vvm_visitor : public ir_hierarchical_visitor
 		}
 	}
 
-	void DumpOps()
+	FString DumpOps(uint32& OutNumOps)
 	{
-#if VM_VERBOSE_LOGGING
 		FString OpStr;
+		OutNumOps = ordered_ops.Num();
 
-		//Dump the constant table
-		OpStr += TEXT("\n-------------------------------\n");
-		OpStr += TEXT("Constant Table\n");
-		OpStr += TEXT("-------------------------------\n");
-		//First go through all the parameters.
 
+		//First go through all the parameters.
+		FString OpsConstantTable;
+		uint32 NumConstants = 0;
 		TFunction<void(variable_info_node*, FString)> GatherConstantTable = [&](variable_info_node* node, FString accumulated_name)
 		{
 			accumulated_name += node->name;
 			if (node->is_scalar())
 			{
 				check(node->offset != INDEX_NONE);
-				OpStr += FString::Printf(TEXT("%d | %s\n"), node->offset, *accumulated_name);
+				OpsConstantTable += FString::Printf(TEXT("%d | %s\n"), node->offset, *accumulated_name);
+				NumConstants++;
 			}
 			else
 			{
@@ -1850,26 +1916,42 @@ class ir_gen_vvm_visitor : public ir_hierarchical_visitor
 			case EVectorVMBaseTypes::Float:
 			{
 				float Val = *(float*)(CompilationOutput.InternalConstantData.GetData() + Offset);
-				OpStr += FString::Printf(TEXT("%d | %f\n"), TableOffset, Val);
+				OpsConstantTable += FString::Printf(TEXT("%d | %f\n"), TableOffset, Val);
+				NumConstants++;
 			}
 			break;
 			case EVectorVMBaseTypes::Int:
 			{
 				int32 Val = *(int32*)(CompilationOutput.InternalConstantData.GetData() + Offset);
-				OpStr += FString::Printf(TEXT("%d | %d\n"), TableOffset, Val);
+				OpsConstantTable += FString::Printf(TEXT("%d | %d\n"), TableOffset, Val);
+				NumConstants++;
 			}
 			break;
 			case EVectorVMBaseTypes::Bool:
 			{
 				int32 Val = *(int32*)(CompilationOutput.InternalConstantData.GetData() + Offset);
-				OpStr += FString::Printf(TEXT("%d | %s\n"), TableOffset, Val == INDEX_NONE ? TEXT("True") : (Val == 0 ? TEXT("False") : TEXT("Invalid") ));
+				OpsConstantTable += FString::Printf(TEXT("%d | %s\n"), TableOffset, Val == INDEX_NONE ? TEXT("True") : (Val == 0 ? TEXT("False") : TEXT("Invalid") ));
+				NumConstants++;
 			}
 			break;
 			}
 		}
+		
+		
+		OpStr += TEXT("\n-------------------------------\n");
+		OpStr += TEXT("Summary\n");
+		OpStr += TEXT("-------------------------------\n");
+		OpStr += FString::Printf(TEXT("Num Byte Code Ops: %d\n"), ordered_ops.Num());
+		OpStr += FString::Printf(TEXT("Num Constants: %d\n"), NumConstants);
+		
+		//Dump the constant table
+		OpStr += TEXT("\n-------------------------------\n");
+		OpStr += TEXT("Constant Table\n");
+		OpStr += TEXT("-------------------------------\n");
+		OpStr += OpsConstantTable;
 
 		OpStr += TEXT("-------------------------------\n");
-		OpStr += TEXT("Byte Code\n");
+		OpStr += FString::Printf(TEXT("Byte Code (%d Ops)\n"), ordered_ops.Num());
 		OpStr += TEXT("-------------------------------\n");
 		
 		//Dump the bytecode		
@@ -1880,8 +1962,10 @@ class ir_gen_vvm_visitor : public ir_hierarchical_visitor
 		
 		OpStr += TEXT("-------------------------------\n");
 
+#if VM_VERBOSE_LOGGING
 		UE_LOG(LogVVMBackend, Log, TEXT("\n%s"), *OpStr);
 #endif
+		return OpStr;
 	}
 
 public:
@@ -1901,7 +1985,7 @@ public:
 		visit_list_elements(&genv, ir);
 		genv.Finalize();
 
-		genv.DumpOps();
+		InCompOutput.AssemblyAsString = genv.DumpOps(InCompOutput.NumOps);
 
 		return NULL;
 	}
diff --git a/Engine/Source/Developer/ShaderFormatVectorVM/Private/ir_vm_merge_op_visitor.cpp b/Engine/Source/Developer/ShaderFormatVectorVM/Private/ir_vm_merge_op_visitor.cpp
--- a/Engine/Source/Developer/ShaderFormatVectorVM/Private/ir_vm_merge_op_visitor.cpp
+++ b/Engine/Source/Developer/ShaderFormatVectorVM/Private/ir_vm_merge_op_visitor.cpp
@@ -96,7 +96,7 @@ public:
 					ir_assignment* assign = mul_assignments[mul_idx];
 					if (assign->lhs->type->is_float())
 					{
-						ir_rvalue* mul_operand = assign->rhs->clone(ralloc_parent(expr), NULL);
+						ir_rvalue* mul_operand = assign->rhs->clone(state, NULL);
 						ir_rvalue* add_operand = equiv_operand == 0 ? expr->operands[1] : expr->operands[0];
 
 						expr->operands[0] = mul_operand;
@@ -144,7 +144,7 @@ public:
 						inner_val = neg_expr->operands[0];
 						check(inner_val);
 
-						ir_rvalue* neg_operand = inner_val->clone(ralloc_parent(expr), NULL);
+						ir_rvalue* neg_operand = inner_val->clone(state, NULL);
 						ir_rvalue* add_operand = equiv_operand == 0 ? expr->operands[1] : expr->operands[0];
 
 						expr->operation = ir_binop_sub;
diff --git a/Engine/Source/Developer/ShaderFormatVectorVM/Private/ir_vm_propagate_non_expressions_visitor.cpp b/Engine/Source/Developer/ShaderFormatVectorVM/Private/ir_vm_propagate_non_expressions_visitor.cpp
--- a/Engine/Source/Developer/ShaderFormatVectorVM/Private/ir_vm_propagate_non_expressions_visitor.cpp
+++ b/Engine/Source/Developer/ShaderFormatVectorVM/Private/ir_vm_propagate_non_expressions_visitor.cpp
@@ -111,7 +111,7 @@ public:
 					check(assign->rhs->as_expression() == nullptr);
 					check(assign->rhs->as_swizzle() || assign->rhs->as_dereference_variable() || assign->rhs->as_constant());
 					check(assign->rhs->type->is_scalar());//All assignments must be scalar at this point!
-					ir_rvalue* new_rval = assign->rhs->clone(ralloc_parent(assign), nullptr);
+					ir_rvalue* new_rval = assign->rhs->clone(parse_state, nullptr);
 					(*rvalue) = new_rval;
 					progress = true;
 				}
@@ -204,78 +204,6 @@ void vm_propagate_non_expressions_visitor(exec_list* ir, _mesa_glsl_parse_state*
 	ir_propagate_non_expressions_visitor::run(ir, state);
 }
 
-//////////////////////////////////////////////////////////////////////////
-/** Replaces any array accesses to matrices with the equivalant swizzle. */
-class ir_matrix_array_access_to_swizzles final : public ir_rvalue_visitor
-{
-public:
-
-	bool progress;
-	ir_matrix_array_access_to_swizzles()
-	{
-		progress = false;
-	}
-
-	virtual ~ir_matrix_array_access_to_swizzles()
-	{
-	}
-
-	unsigned get_component_from_matrix_array_deref(ir_dereference_array* array_deref)
-	{
-		check(array_deref);
-		check(array_deref->variable_referenced()->type->is_matrix());
-		ir_constant* index = array_deref->array_index->as_constant();
-		check((index->type == glsl_type::uint_type || index->type == glsl_type::int_type) && index->type->is_scalar());
-		unsigned deref_idx = index->type == glsl_type::uint_type ? index->value.u[0] : index->value.i[0];
-		return deref_idx * array_deref->variable_referenced()->type->vector_elements;
-	}
-
-	virtual void handle_rvalue(ir_rvalue **rvalue)
-	{
-		if (rvalue && *rvalue)
-		{
-			if (ir_swizzle* swiz = (*rvalue)->as_swizzle())
-			{
-				check(swiz->type->is_scalar());//Post scalarize so this should be all scalar.				
-				if (ir_dereference_array* array_deref = swiz->val->as_dereference_array())
-				{
-					swiz->mask.x += get_component_from_matrix_array_deref(array_deref);
-					swiz->val = array_deref->array->clone(ralloc_parent(swiz), nullptr);
-				}
-			}
-		}
-	}
-	
-	virtual ir_visitor_status visit_leave(ir_assignment *assign)
-	{		
-		//Remove matrix array access on the lhs too.
-		if (ir_dereference_array* array_deref = assign->lhs->as_dereference_array())
-		{			
-			assign->write_mask <<= get_component_from_matrix_array_deref(array_deref);
-			assign->lhs = array_deref->array->clone(ralloc_parent(assign), nullptr)->as_dereference();
-		}
-		return ir_rvalue_visitor::visit_leave(assign);
-	}
-
-	static void run(exec_list *ir)
-	{
-		bool progress = false;
-		do
-		{
-			ir_matrix_array_access_to_swizzles v;
-			visit_list_elements(&v, ir);
-
-			progress = v.progress;
-
-			//vm_debug_print("== propagate non expressions - BEFORE DEADCODE ==\n");
-			//vm_debug_dump(ir, state);
-
-			//progress = do_dead_code(ir, false) || progress;
-			//progress = do_dead_code_local(ir) || progress;
-		} while (progress);
-	}
-};
-
 struct MatrixVectors
 {
 	ir_variable* v[4];
@@ -284,12 +212,14 @@ struct MatrixVectors
 class ir_matrices_to_vectors : public ir_rvalue_visitor
 {
 	bool progress;
+	_mesa_glsl_parse_state* parse_state;
 
 	TMap<class ir_variable*, MatrixVectors> MatrixVectorMap;
 public:
-	ir_matrices_to_vectors()
+	ir_matrices_to_vectors(_mesa_glsl_parse_state* in_parse_state)
 	{
 		progress = false;
+		parse_state = in_parse_state;
 	}
 
 	virtual ~ir_matrices_to_vectors()
@@ -300,17 +230,27 @@ public:
 	{
 		if (rvalue && *rvalue)
 		{		
-			check((*rvalue)->as_dereference_array() == nullptr);
-			if (ir_swizzle* swiz = (*rvalue)->as_swizzle())
+			if (ir_dereference* deref = (*rvalue)->as_dereference_array())
 			{
-				ir_variable* var = swiz->variable_referenced();
-				if (var->type->is_matrix())
+				ir_variable* var = deref->variable_referenced();
+				if (var->type->is_matrix() && var->mode != ir_var_uniform)
 				{
 					MatrixVectors& mv = MatrixVectorMap.FindChecked(var);
-					unsigned v_idx = swiz->mask.x / 4;
-					swiz->mask.x %= 4;					
-					void* p = ralloc_parent(swiz);
-					swiz->val = new(p) ir_dereference_variable(mv.v[v_idx]);
+					if (ir_dereference_array* array_deref = (*rvalue)->as_dereference_array())
+					{
+						ir_constant* idx = array_deref->array_index->as_constant();
+						check(idx);
+						unsigned v_idx = idx->value.i[0];
+						//void* p = ralloc_parent(array_deref);
+						*rvalue = new(parse_state) ir_dereference_variable(mv.v[v_idx]);
+					}
+					else if (ir_swizzle* swiz = (*rvalue)->as_swizzle())
+					{
+						unsigned v_idx = swiz->mask.x / 4;
+						swiz->mask.x %= 4;
+						//void* p = ralloc_parent(swiz);
+						swiz->val = new(parse_state) ir_dereference_variable(mv.v[v_idx]);
+					}
 				}
 			}
 		}	
@@ -318,27 +258,55 @@ public:
 
 	virtual ir_visitor_status visit_leave(ir_assignment *assign)
 	{
+		//Dont think this is required.
+		//Previous pass should convert all matrix ops to vec ops so the worst this will be is a Mat[row] = somevec.xyzw so replacing the array deref is is fine.
 		ir_variable* var = assign->lhs->variable_referenced();
 		if (var->type->is_matrix())
 		{
 			MatrixVectors& mv = MatrixVectorMap.FindChecked(var);
-			ir_variable* v = nullptr;
-			unsigned v_idx = 0;
-			while (assign->write_mask > 0)
+			if (ir_dereference_array* array_deref = assign->lhs->as_dereference_array())
+			{
+				ir_constant* idx = array_deref->array_index->as_constant();
+				check(idx);
+				unsigned v_idx = idx->value.i[0];
+				//void* p = ralloc_parent(array_deref);
+				assign->set_lhs(new(parse_state) ir_dereference_variable(mv.v[v_idx]));
+			}
+			else if (ir_dereference_variable* matrix_deref = assign->lhs->as_dereference_variable())
 			{
-				if ((assign->write_mask & 15) != 0)
+				ir_dereference_variable* src_mat = assign->rhs->as_dereference_variable();
+				check(src_mat && src_mat->type->is_matrix());
+				ir_dereference* src_derefs[4];
+				//void* p = ralloc_parent(array_deref);
+				if (src_mat->variable_referenced()->mode == ir_var_uniform)
+				{
+					src_derefs[0] = new(parse_state) ir_dereference_array(src_mat->variable_referenced(), new(parse_state) ir_constant(0));
+					src_derefs[1] = new(parse_state) ir_dereference_array(src_mat->variable_referenced(), new(parse_state) ir_constant(1));
+					src_derefs[2] = new(parse_state) ir_dereference_array(src_mat->variable_referenced(), new(parse_state) ir_constant(2));
+					src_derefs[3] = new(parse_state) ir_dereference_array(src_mat->variable_referenced(), new(parse_state) ir_constant(3));
+
+				}
+				else
 				{
-					v = mv.v[v_idx];
-					check(assign->write_mask >> 4 == 0);
-					break;
+					MatrixVectors& src_mv = MatrixVectorMap.FindChecked(src_mat->variable_referenced());
+					src_derefs[0] = new(parse_state) ir_dereference_variable(src_mv.v[0]);
+					src_derefs[1] = new(parse_state) ir_dereference_variable(src_mv.v[1]);
+					src_derefs[2] = new(parse_state) ir_dereference_variable(src_mv.v[2]);
+					src_derefs[3] = new(parse_state) ir_dereference_variable(src_mv.v[3]);
 				}
-				++v_idx;
-				assign->write_mask >>= 4;
+
+				assign->insert_before(new(parse_state) ir_assignment(new(parse_state) ir_dereference_variable(mv.v[0]), src_derefs[0], assign->condition, 0xF));
+				assign->insert_before(new(parse_state) ir_assignment(new(parse_state) ir_dereference_variable(mv.v[1]), src_derefs[1], assign->condition, 0xF));
+				assign->insert_before(new(parse_state) ir_assignment(new(parse_state) ir_dereference_variable(mv.v[2]), src_derefs[2], assign->condition, 0xF));
+				assign->insert_before(new(parse_state) ir_assignment(new(parse_state) ir_dereference_variable(mv.v[3]), src_derefs[3], assign->condition, 0xF));
+				assign->remove();
+			}
+			else
+			{
+				check(false);//Should have removed matrix swizzles by now
 			}
-			check(v);
-			void* p = ralloc_parent(assign);
-			assign->lhs = new(p) ir_dereference_variable(v);
 		}
+
 		return ir_rvalue_visitor::visit_leave(assign);
 	}
 
@@ -352,22 +320,27 @@ public:
 				mv = &MatrixVectorMap.Add(var);
 				void* p = ralloc_parent(var);
 
-				check(var->next && var->prev);
-
 				const glsl_type* vtype = var->type->column_type();
-				const char *name = ralloc_asprintf(p, "%s_%s", var->name,"col0");				
+
+				const char* base_name = var->name ? var->name : "temp";
+				const char *name = ralloc_asprintf(p, "%s_%s", base_name,"col0");
 				mv->v[0] = new(p) ir_variable(vtype, name, var->mode);
-				var->insert_before(mv->v[0]);
-				name = ralloc_asprintf(p, "%s_%s", var->name, "col1");
+				name = ralloc_asprintf(p, "%s_%s", base_name, "col1");
 				mv->v[1] = new(p) ir_variable(vtype, name, var->mode);
-				var->insert_before(mv->v[1]);
-				name = ralloc_asprintf(p, "%s_%s", var->name, "col2");
+				name = ralloc_asprintf(p, "%s_%s", base_name, "col2");
 				mv->v[2] = new(p) ir_variable(vtype, name, var->mode);
-				var->insert_before(mv->v[2]);
-				name = ralloc_asprintf(p, "%s_%s", var->name, "col3");
+				name = ralloc_asprintf(p, "%s_%s", base_name, "col3");
 				mv->v[3] = new(p) ir_variable(vtype, name, var->mode);
-				var->insert_before(mv->v[3]);
-				var->remove();
+
+				//This is called manually on uniforms too but these are not in a cclist so those functions need not (and cannot) be called.
+				if (var->mode != ir_var_uniform)
+				{
+					var->insert_before(mv->v[0]);
+					var->insert_before(mv->v[1]);
+					var->insert_before(mv->v[2]);
+					var->insert_before(mv->v[3]);
+					var->remove();
+				}
 			}
 		}
 		return visit_continue;
@@ -378,16 +351,7 @@ public:
 		bool progress = false;
 		do
 		{
-			ir_matrices_to_vectors v;
-			
-			//Have to manually visit uniforms first.
-			for (SCBuffer& buff : state->CBuffersOriginal)
-			{
-				for (SCBufferMember& member : buff.Members)
-				{
-					member.Var->accept(&v);
-				}
-			}
+			ir_matrices_to_vectors v(state);
 
 			visit_list_elements(&v, ir);
 
@@ -404,6 +368,5 @@ public:
 
 void vm_matrices_to_vectors(exec_list* ir, _mesa_glsl_parse_state* state)
 {	
-	ir_matrix_array_access_to_swizzles::run(ir);
 	ir_matrices_to_vectors::run(ir, state);
 }
diff --git a/Engine/Source/Developer/ShaderFormatVectorVM/Private/ir_vm_scalarize_visitor.cpp b/Engine/Source/Developer/ShaderFormatVectorVM/Private/ir_vm_scalarize_visitor.cpp
--- a/Engine/Source/Developer/ShaderFormatVectorVM/Private/ir_vm_scalarize_visitor.cpp
+++ b/Engine/Source/Developer/ShaderFormatVectorVM/Private/ir_vm_scalarize_visitor.cpp
@@ -124,7 +124,7 @@ class ir_scalarize_visitor2 : public ir_hierarchical_visitor
 		//Ugh. Make less bad.
 		ir_function* func = const_cast<ir_function*>(in_sig->function());
 
-		void* parent = ralloc_parent(in_sig);
+		//void* parent = ralloc_parent(in_sig);
 		ECallScalarizeMode mode = get_scalarize_mode(in_sig);
 
 		ir_function_signature* new_sig = nullptr;
@@ -134,14 +134,14 @@ class ir_scalarize_visitor2 : public ir_hierarchical_visitor
 		}
 		else if (mode == ECallScalarizeMode::SplitParams)
 		{
-			new_sig = new(parent) ir_function_signature(in_sig->return_type);
+			new_sig = new(parse_state) ir_function_signature(in_sig->return_type);
 
 			TFunction<void(const ir_variable*, const glsl_type*, FString)> append_scalar_params = [&](const ir_variable* original, const glsl_type* type, FString Name)
 			{
 				const glsl_type* base_type = type->get_base_type();
 				if (type->is_scalar())
 				{
-					new_sig->parameters.push_tail(new(parent) ir_variable(type, TCHAR_TO_ANSI(*Name), original->mode));
+					new_sig->parameters.push_tail(new(parse_state) ir_variable(type, TCHAR_TO_ANSI(*Name), original->mode));
 				}
 				else if (type->is_vector())
 				{
@@ -187,11 +187,11 @@ class ir_scalarize_visitor2 : public ir_hierarchical_visitor
 		}
 		else if (mode == ECallScalarizeMode::SplitCalls)
 		{
-			new_sig = new(parent) ir_function_signature(in_sig->return_type);
+			new_sig = new(parse_state) ir_function_signature(in_sig->return_type->get_base_type());
 			foreach_iter(exec_list_iterator, iter, in_sig->parameters)
 			{
 				ir_variable *var = (ir_variable *)iter.get();
-				new_sig->parameters.push_tail(new(parent) ir_variable(var->type->get_base_type(), var->name, var->mode));
+				new_sig->parameters.push_tail(new(parse_state) ir_variable(var->type->get_base_type(), var->name, var->mode));
 			}
 		}
 
@@ -255,7 +255,7 @@ class ir_scalarize_visitor2 : public ir_hierarchical_visitor
 	
 	virtual ir_visitor_status visit_enter(ir_call* call)
 	{
-		void* parent = ralloc_parent(call->callee);
+		//void* parent = ralloc_parent(call->callee);
 		ECallScalarizeMode mode = get_scalarize_mode(call->callee);
 
 		check(call->next && call->prev);
@@ -316,29 +316,29 @@ class ir_scalarize_visitor2 : public ir_hierarchical_visitor
 
 				//Clone the call max_component times and visit each param to scalarize them to the right component.				
 				ir_variable* old_dest = call->return_deref->var;
-				void* perm_mem_ctx = ralloc_parent(call);
+				//void* perm_mem_ctx = ralloc_parent(call);
 				for (unsigned i = 0; i < max_components; ++i)
 				{
 					dest_component = i;
 
-					ir_variable* new_dest = new(perm_mem_ctx) ir_variable(old_dest->type->get_base_type(), old_dest->name, ir_var_temporary);
+					ir_variable* new_dest = new(parse_state) ir_variable(old_dest->type->get_base_type(), old_dest->name, ir_var_temporary);
 					call->insert_before(new_dest);
 
 					exec_list new_params;
 					foreach_iter(exec_list_iterator, param_iter, call->actual_parameters)
 					{
 						ir_rvalue* param = (ir_rvalue*)param_iter.get();
-						param = param->clone(perm_mem_ctx, nullptr);
+						param = param->clone(parse_state, nullptr);
 						curr_rval = nullptr;
 						param->accept(this);
 						new_params.push_tail(curr_rval ? curr_rval : param);
 					}
 
-					call->insert_before(new(perm_mem_ctx)ir_call(scalar_sig, new(perm_mem_ctx) ir_dereference_variable(new_dest), &new_params));
+					call->insert_before(new(parse_state)ir_call(scalar_sig, new(parse_state) ir_dereference_variable(new_dest), &new_params));
 
-					ir_dereference* deref = call->return_deref->clone(perm_mem_ctx, nullptr);
+					ir_dereference* deref = call->return_deref->clone(parse_state, nullptr);
 					unsigned write_mask = 1 << dest_component;
-					call->insert_before(new(perm_mem_ctx)ir_assignment(deref, new(perm_mem_ctx) ir_dereference_variable(new_dest), nullptr, write_mask));
+					call->insert_before(new(parse_state)ir_assignment(deref, new(parse_state) ir_dereference_variable(new_dest), nullptr, write_mask));
 				}
 
 				call->remove();
@@ -363,7 +363,7 @@ class ir_scalarize_visitor2 : public ir_hierarchical_visitor
 
 		check(assign->next && assign->prev);
 
-		void* perm_mem_ctx = ralloc_parent(assign);
+		//void* perm_mem_ctx = ralloc_parent(assign);
 
 		const glsl_type* type = assign->lhs->type;
 
@@ -401,14 +401,14 @@ class ir_scalarize_visitor2 : public ir_hierarchical_visitor
 					comp_assign = NULL;
 				}
 
-				comp_assign = assign->clone(perm_mem_ctx, NULL);
+				comp_assign = assign->clone(parse_state, NULL);
 				check(comp_assign);
 				dest_component = comp_idx;
 
 				if (is_struct)
 				{
 					comp_assign->write_mask = 0;
-					comp_assign->set_lhs(new(perm_mem_ctx)ir_dereference_record(comp_assign->lhs, comp_assign->lhs->type->fields.structure[dest_component].name));
+					comp_assign->set_lhs(new(parse_state)ir_dereference_record(comp_assign->lhs, comp_assign->lhs->type->fields.structure[dest_component].name));
 				}
 				else
 				{
@@ -466,22 +466,26 @@ class ir_scalarize_visitor2 : public ir_hierarchical_visitor
 
 	virtual ir_visitor_status visit_enter(ir_dereference_array* array_deref)
 	{
-		void* perm_mem_ctx = ralloc_parent(array_deref);
+		//void* perm_mem_ctx = ralloc_parent(array_deref);
 		const glsl_type* array_type = array_deref->array->type;
 		const glsl_type* type = array_deref->type;
 
 		//Only supporting array derefs for matrices atm.
 		check(array_type->is_matrix());
+		check(type->is_vector());
 
 		//Only support constant matrix indices. Not immediately clear how I'd do non-const access.
 		ir_constant* index = array_deref->array_index->as_constant();
 		check(index->type == glsl_type::uint_type || index->type == glsl_type::int_type);
 		check(index->type->is_scalar());
 
-		unsigned mat_comp = index->type == glsl_type::uint_type ? index->value.u[0] : index->value.i[0];
-		unsigned swiz_comp = mat_comp * array_type->vector_elements + dest_component;
-		ir_dereference_variable* new_deref = new(perm_mem_ctx) ir_dereference_variable(array_deref->variable_referenced());
-		ir_swizzle* swiz = new(perm_mem_ctx) ir_swizzle(new_deref, type->is_scalar() ? 0 : swiz_comp, 0, 0, 0, 1);
+// 		unsigned mat_comp = index->type == glsl_type::uint_type ? index->value.u[0] : index->value.i[0];
+// 		unsigned swiz_comp = mat_comp * array_type->vector_elements + dest_component;
+// 		ir_dereference_variable* new_deref = new(perm_mem_ctx) ir_dereference_variable(array_deref->variable_referenced());
+// 		ir_swizzle* swiz = new(perm_mem_ctx) ir_swizzle(new_deref, type->is_scalar() ? 0 : swiz_comp, 0, 0, 0, 1);
+// 		curr_rval = swiz;
+
+		ir_swizzle* swiz = new(parse_state) ir_swizzle(array_deref, dest_component, 0, 0, 0, 1);
 		curr_rval = swiz;
 
 		return visit_continue_with_parent;
@@ -489,14 +493,14 @@ class ir_scalarize_visitor2 : public ir_hierarchical_visitor
 
 	virtual ir_visitor_status visit_enter(ir_dereference_record* deref)
 	{
-		void* perm_mem_ctx = ralloc_parent(deref);
+		//void* perm_mem_ctx = ralloc_parent(deref);
 		const glsl_type* type = deref->type;
 
 		if (type->base_type == GLSL_TYPE_STRUCT)
 		{
 			check(dest_component < type->length);
 
-			ir_dereference_record* rec = new(perm_mem_ctx) ir_dereference_record(deref, type->fields.structure[dest_component].name);
+			ir_dereference_record* rec = new(parse_state) ir_dereference_record(deref, type->fields.structure[dest_component].name);
 			curr_rval = rec;
 		}
 		else
@@ -505,7 +509,7 @@ class ir_scalarize_visitor2 : public ir_hierarchical_visitor
 			//check(type->components() >= dest_component || type->is_scalar());
 			unsigned use_component = FMath::Min(type->components(), dest_component);
 
-			ir_swizzle* swiz = new(perm_mem_ctx) ir_swizzle(deref, type->is_scalar() ? 0 : use_component, 0, 0, 0, 1);
+			ir_swizzle* swiz = new(parse_state) ir_swizzle(deref, type->is_scalar() ? 0 : use_component, 0, 0, 0, 1);
 			curr_rval = swiz;
 		}
 
@@ -514,14 +518,14 @@ class ir_scalarize_visitor2 : public ir_hierarchical_visitor
 
 	virtual ir_visitor_status visit(ir_dereference_variable *deref)
 	{
-		void* perm_mem_ctx = ralloc_parent(deref);
+		//void* perm_mem_ctx = ralloc_parent(deref);
 		ir_variable* var = deref->variable_referenced();
 		const glsl_type* type = var->type;
 		if (type->base_type == GLSL_TYPE_STRUCT)
 		{
 			check(dest_component < type->length);
 
-			ir_dereference_record* rec = new(perm_mem_ctx) ir_dereference_record(var, type->fields.structure[dest_component].name);
+			ir_dereference_record* rec = new(parse_state) ir_dereference_record(var, type->fields.structure[dest_component].name);
 			curr_rval = rec;
 		}
 		else
@@ -532,7 +536,7 @@ class ir_scalarize_visitor2 : public ir_hierarchical_visitor
 				//check(type->components() > dest_component);
 				unsigned use_component = FMath::Min(type->components(), dest_component);
 
-				ir_swizzle* swiz = new(perm_mem_ctx) ir_swizzle(deref, use_component, 0, 0, 0, 1);
+				ir_swizzle* swiz = new(parse_state) ir_swizzle(deref, use_component, 0, 0, 0, 1);
 				curr_rval = swiz;
 			}
 		}
diff --git a/Engine/Source/Developer/ShaderFormatVectorVM/Private/ir_vm_to_single_op_visitor.cpp b/Engine/Source/Developer/ShaderFormatVectorVM/Private/ir_vm_to_single_op_visitor.cpp
--- a/Engine/Source/Developer/ShaderFormatVectorVM/Private/ir_vm_to_single_op_visitor.cpp
+++ b/Engine/Source/Developer/ShaderFormatVectorVM/Private/ir_vm_to_single_op_visitor.cpp
@@ -81,16 +81,16 @@ public:
 
 		unsigned expr_depth = expression_stack.Num();
 		ir_instruction* parent = expr_depth > 0 ? expression_stack.Top() : nullptr;
-		void* parent_ctx = ralloc_parent(expr);
+		//void* parent_ctx = ralloc_parent(expr);
 
 		if (parent)
 		{
-			ir_variable *tmp_var = new(parent_ctx) ir_variable(expr->type, "tmp_var", ir_var_temporary);
-			ir_assignment *tmp_assign = new(parent_ctx) ir_assignment(new(parent_ctx) ir_dereference_variable(tmp_var), expr);
+			ir_variable *tmp_var = new(state) ir_variable(expr->type, "tmp_var", ir_var_temporary);
+			ir_assignment *tmp_assign = new(state) ir_assignment(new(state) ir_dereference_variable(tmp_var), expr);
 			check(tmp_assign->write_mask > 0);
 			base_ir->insert_before(tmp_var);
 			base_ir->insert_before(tmp_assign);
-			replacement = new(parent_ctx) ir_dereference_variable(tmp_var);
+			replacement = new(state) ir_dereference_variable(tmp_var);
 		}
 
 		return visit_continue_with_parent;
@@ -105,13 +105,13 @@ public:
 		check(replacement == nullptr);
 		if (assign->lhs->variable_referenced()->mode == ir_var_out && assign_has_expressions)
 		{
-			void *mem_ctx = ralloc_parent(assign);
-			ir_variable *tmp_var = new(mem_ctx) ir_variable(assign->rhs->type, "output_var", ir_var_temporary);
-			ir_assignment *tmp_assign = new(mem_ctx) ir_assignment(new(mem_ctx) ir_dereference_variable(tmp_var), assign->rhs);
+			//void *mem_ctx = ralloc_parent(assign);
+			ir_variable *tmp_var = new(state) ir_variable(assign->rhs->type, "output_var", ir_var_temporary);
+			ir_assignment *tmp_assign = new(state) ir_assignment(new(state) ir_dereference_variable(tmp_var), assign->rhs);
 			check(tmp_assign->write_mask > 0);
 			base_ir->insert_before(tmp_var);
 			base_ir->insert_before(tmp_assign);
-			assign->rhs = new(mem_ctx) ir_dereference_variable(tmp_var);
+			assign->rhs = new(state) ir_dereference_variable(tmp_var);
 			progress = true;
 		}
 
@@ -122,7 +122,7 @@ public:
 	{
 		check(base_ir->next && base_ir->prev);
 		assign_has_expressions = true;
-		void *mem_ctx = ralloc_parent(call);
+		//void *mem_ctx = ralloc_parent(call);
 
 		foreach_iter(exec_list_iterator, iter, call->actual_parameters)
 		{
@@ -132,12 +132,12 @@ public:
 			if (!deref)
 			{
 				//This param isn't just a deref so move it out to it's own variable and deref it in the call instead.	
-				ir_variable *tmp_var = new(mem_ctx) ir_variable(param->type, "call_param_temp", ir_var_temporary);
-				ir_assignment *tmp_assign = new(mem_ctx) ir_assignment(new(mem_ctx) ir_dereference_variable(tmp_var), param);
+				ir_variable *tmp_var = new(state) ir_variable(param->type, "call_param_temp", ir_var_temporary);
+				ir_assignment *tmp_assign = new(state) ir_assignment(new(state) ir_dereference_variable(tmp_var), param);
 				check(tmp_assign->write_mask > 0);
 				base_ir->insert_before(tmp_var);
 				base_ir->insert_before(tmp_assign);
-				ir_rvalue* new_param = new(mem_ctx) ir_dereference_variable(tmp_var);
+				ir_rvalue* new_param = new(state) ir_dereference_variable(tmp_var);
 				check(param->next && param->prev);
 				param->replace_with(new_param);
 				param = new_param;
diff --git a/Engine/Source/Developer/ShaderFormatVectorVM/Private/ir_vm_vec_op_to_scalar_visitor.cpp b/Engine/Source/Developer/ShaderFormatVectorVM/Private/ir_vm_vec_op_to_scalar_visitor.cpp
--- a/Engine/Source/Developer/ShaderFormatVectorVM/Private/ir_vm_vec_op_to_scalar_visitor.cpp
+++ b/Engine/Source/Developer/ShaderFormatVectorVM/Private/ir_vm_vec_op_to_scalar_visitor.cpp
@@ -25,9 +25,12 @@ PRAGMA_POP
 //Additional helper during scalaization to split any special vector expressions such as dot into their scalar operations.
 class ir_vec_op_to_scalar_visitor : public ir_hierarchical_visitor
 {
+	_mesa_glsl_parse_state* state;
+
 public:
-	ir_vec_op_to_scalar_visitor()
+	ir_vec_op_to_scalar_visitor(_mesa_glsl_parse_state* in_state)
 	{
+		state = in_state;
 		this->made_progress = false;
 	}
 
@@ -37,7 +40,7 @@ public:
 
 		check(base_ir->next && base_ir->prev);
 
-		void* parent = ralloc_parent(base_ir);
+		//void* parent = ralloc_parent(base_ir);
 		/* Avoid making a temporary if we don't need to to avoid aliasing. */
 		if (deref &&
 			deref->variable_referenced() != result->variable_referenced())
@@ -48,7 +51,7 @@ public:
 		/* Otherwise, store the operand in a temporary generally if it's
 		* not a dereference.
 		*/
-		ir_variable *var = new(parent)ir_variable(rval->type,
+		ir_variable *var = new(state)ir_variable(rval->type,
 			"vec_op_to_scalar",
 			ir_var_temporary);
 		base_ir->insert_before(var);
@@ -56,9 +59,9 @@ public:
 		/* Note that we use this dereference for the assignment.  That means
 		* that others that want to use op[i] have to clone the deref.
 		*/
-		deref = new(parent)ir_dereference_variable(var);
+		deref = new(state)ir_dereference_variable(var);
 
-		ir_assignment *assign = new(parent)ir_assignment(deref, rval);
+		ir_assignment *assign = new(state)ir_assignment(deref, rval);
 		check(assign->write_mask > 0);
 		base_ir->insert_before(assign);
 		return deref;
@@ -93,7 +96,7 @@ public:
 
 	ir_visitor_status visit_leave(ir_assignment * orig_assign)
 	{
-		void* parent = ralloc_parent(base_ir);
+		//void* parent = ralloc_parent(base_ir);
 		ir_expression *orig_expr = orig_assign->rhs->as_expression();
 		ir_dereference *op[2] = {nullptr, nullptr};
 
@@ -143,23 +146,23 @@ void ir_vec_op_to_scalar_visitor::do_dot(ir_dereference *result, ir_dereference
 	ir_expression *expr;
 	check(base_ir->next && base_ir->prev);
 	check(a->type == b->type);
-	void* parent = ralloc_parent(base_ir);
+	//void* parent = ralloc_parent(base_ir);
 
-	expr = new(parent)ir_expression(ir_binop_mul, a->type->get_base_type(),
-		new(parent)ir_swizzle(a, 0, 0, 0, 0, 1),
-		new(parent)ir_swizzle(b, 0, 0, 0, 0, 1));
+	expr = new(state)ir_expression(ir_binop_mul, a->type->get_base_type(),
+		new(state)ir_swizzle(a, 0, 0, 0, 0, 1),
+		new(state)ir_swizzle(b, 0, 0, 0, 0, 1));
 
 	for (unsigned comp = 1; comp < a->type->vector_elements; comp++)
 	{
-		expr = new(parent)ir_expression(ir_binop_add, a->type->get_base_type(),
+		expr = new(state)ir_expression(ir_binop_add, a->type->get_base_type(),
 			expr,
-			new(parent)ir_expression(ir_binop_mul, a->type->get_base_type(),
-				new(parent)ir_swizzle(a, comp, 0, 0, 0, 1),
-				new(parent)ir_swizzle(b, comp, 0, 0, 0, 1))
+			new(state)ir_expression(ir_binop_mul, a->type->get_base_type(),
+				new(state)ir_swizzle(a, comp, 0, 0, 0, 1),
+				new(state)ir_swizzle(b, comp, 0, 0, 0, 1))
 		);
 	}
 
-	assign = new(parent)ir_assignment(result, expr);
+	assign = new(state)ir_assignment(result, expr);
 	check(assign->write_mask > 0);
 	base_ir->insert_before(assign);
 }
@@ -174,23 +177,23 @@ void ir_vec_op_to_scalar_visitor::do_cross(ir_dereference *result, ir_dereferenc
 	const glsl_type* type = a->type;
 	const glsl_type* basetype = a->type->get_base_type();
 
-	void* parent = ralloc_parent(base_ir);
+	//void* parent = ralloc_parent(base_ir);
 
-#define GETX(v) (new(parent)ir_swizzle(v, 0, 0, 0, 0, 1))
-#define GETY(v) (new(parent)ir_swizzle(v, 1, 0, 0, 0, 1))
-#define GETZ(v) (new(parent)ir_swizzle(v, 2, 0, 0, 0, 1))
-#define MUL(A,B)  (new(parent)ir_expression(ir_binop_mul, basetype, A, B))
-#define SUB(A,B)  (new(parent)ir_expression(ir_binop_sub, basetype, A, B))
+#define GETX(v) (new(state)ir_swizzle(v, 0, 0, 0, 0, 1))
+#define GETY(v) (new(state)ir_swizzle(v, 1, 0, 0, 0, 1))
+#define GETZ(v) (new(state)ir_swizzle(v, 2, 0, 0, 0, 1))
+#define MUL(A,B)  (new(state)ir_expression(ir_binop_mul, basetype, A, B))
+#define SUB(A,B)  (new(state)ir_expression(ir_binop_sub, basetype, A, B))
 
-	assign = new(parent)ir_assignment(result, SUB(MUL(GETY(a), GETZ(b)), MUL(GETZ(a), GETY(b))));
+	assign = new(state)ir_assignment(result, SUB(MUL(GETY(a), GETZ(b)), MUL(GETZ(a), GETY(b))));
 	assign->write_mask = 1;
 	base_ir->insert_before(assign);
 
-	assign = new(parent)ir_assignment(result, SUB(MUL(GETZ(a), GETX(b)), MUL(GETX(a), GETZ(b))));
+	assign = new(state)ir_assignment(result, SUB(MUL(GETZ(a), GETX(b)), MUL(GETX(a), GETZ(b))));
 	assign->write_mask = 1 << 1;
 	base_ir->insert_before(assign);
 
-	assign = new(parent)ir_assignment(result, SUB(MUL(GETX(a), GETY(b)), MUL(GETY(a), GETX(b))));
+	assign = new(state)ir_assignment(result, SUB(MUL(GETX(a), GETY(b)), MUL(GETY(a), GETX(b))));
 	assign->write_mask = 1 << 2;
 	base_ir->insert_before(assign);
 
@@ -208,28 +211,28 @@ void ir_vec_op_to_scalar_visitor::do_normalize(ir_dereference *result, ir_derefe
 	check(base_ir->next && base_ir->prev);
 	check(a->type->is_vector());
 
-	void* parent = ralloc_parent(base_ir);
+	//void* parent = ralloc_parent(base_ir);
 
 	//generate the dot of a with itself.
-	expr = new(parent)ir_expression(ir_binop_mul, a->type->get_base_type(),
-		new(parent)ir_swizzle(a, 0, 0, 0, 0, 1),
-		new(parent)ir_swizzle(a, 0, 0, 0, 0, 1));
+	expr = new(state)ir_expression(ir_binop_mul, a->type->get_base_type(),
+		new(state)ir_swizzle(a, 0, 0, 0, 0, 1),
+		new(state)ir_swizzle(a, 0, 0, 0, 0, 1));
 
 	for (unsigned comp = 1; comp < a->type->vector_elements; comp++)
 	{
-		expr = new(parent)ir_expression(ir_binop_add, a->type->get_base_type(),
+		expr = new(state)ir_expression(ir_binop_add, a->type->get_base_type(),
 			expr,
-			new(parent)ir_expression(ir_binop_mul, a->type->get_base_type(),
-				new(parent)ir_swizzle(a, comp, 0, 0, 0, 1),
-				new(parent)ir_swizzle(a, comp, 0, 0, 0, 1))
+			new(state)ir_expression(ir_binop_mul, a->type->get_base_type(),
+				new(state)ir_swizzle(a, comp, 0, 0, 0, 1),
+				new(state)ir_swizzle(a, comp, 0, 0, 0, 1))
 		);
 	}
 
 	//Generate the length and divide a by it.
-	ir_expression* inv_len = new(parent)ir_expression(ir_unop_rsq, a->type->get_base_type(), expr);
-	expr = new(parent)ir_expression(ir_binop_mul, a->type, a, inv_len);
+	ir_expression* inv_len = new(state)ir_expression(ir_unop_rsq, a->type->get_base_type(), expr);
+	expr = new(state)ir_expression(ir_binop_mul, a->type, a, inv_len);
 
-	assign = new(parent)ir_assignment(result, expr);
+	assign = new(state)ir_assignment(result, expr);
 	check(assign->write_mask > 0);
 	base_ir->insert_before(assign);
 }
@@ -241,25 +244,25 @@ void ir_vec_op_to_scalar_visitor::do_length(ir_dereference *result, ir_dereferen
 	check(base_ir->next && base_ir->prev);
 	check(a->type->is_vector());
 
-	void* parent = ralloc_parent(base_ir);
+	//void* parent = ralloc_parent(base_ir);
 
 	//generate the dot of a with itself.
-	expr = new(parent)ir_expression(ir_binop_mul, a->type->get_base_type(),
-		new(parent)ir_swizzle(a, 0, 0, 0, 0, 1),
-		new(parent)ir_swizzle(a, 0, 0, 0, 0, 1));
+	expr = new(state)ir_expression(ir_binop_mul, a->type->get_base_type(),
+		new(state)ir_swizzle(a, 0, 0, 0, 0, 1),
+		new(state)ir_swizzle(a, 0, 0, 0, 0, 1));
 
 	for (unsigned comp = 1; comp < a->type->vector_elements; comp++)
 	{
-		expr = new(parent)ir_expression(ir_binop_add, a->type->get_base_type(),
+		expr = new(state)ir_expression(ir_binop_add, a->type->get_base_type(),
 			expr,
-			new(parent)ir_expression(ir_binop_mul, a->type->get_base_type(),
-				new(parent)ir_swizzle(a, comp, 0, 0, 0, 1),
-				new(parent)ir_swizzle(a, comp, 0, 0, 0, 1))
+			new(state)ir_expression(ir_binop_mul, a->type->get_base_type(),
+				new(state)ir_swizzle(a, comp, 0, 0, 0, 1),
+				new(state)ir_swizzle(a, comp, 0, 0, 0, 1))
 		);
 	}
 
-	expr = new(parent)ir_expression(ir_unop_sqrt, a->type->get_base_type(), expr);
-	assign = new(parent)ir_assignment(result, expr);
+	expr = new(state)ir_expression(ir_unop_sqrt, a->type->get_base_type(), expr);
+	assign = new(state)ir_assignment(result, expr);
 	check(assign->write_mask > 0);
 	base_ir->insert_before(assign);
 }
@@ -282,9 +285,9 @@ static bool vec_op_to_scalar_predicate(ir_instruction *ir)
 	return false;
 }
 
-bool do_vec_op_to_scalar(exec_list *instructions)
+bool do_vec_op_to_scalar(exec_list *instructions, _mesa_glsl_parse_state* parse_state)
 {
-	ir_vec_op_to_scalar_visitor v;
+	ir_vec_op_to_scalar_visitor v(parse_state);
 
 	do_expression_flattening(instructions, vec_op_to_scalar_predicate);
 
diff --git a/Engine/Source/Developer/ShaderFormatVectorVM/Public/ShaderFormatVectorVM.h b/Engine/Source/Developer/ShaderFormatVectorVM/Public/ShaderFormatVectorVM.h
--- a/Engine/Source/Developer/ShaderFormatVectorVM/Public/ShaderFormatVectorVM.h
+++ b/Engine/Source/Developer/ShaderFormatVectorVM/Public/ShaderFormatVectorVM.h
@@ -7,7 +7,7 @@
 /** Data which is generated from the hlsl by the VectorVMBackend and fed back into the */
 struct FVectorVMCompilationOutput
 {
-	FVectorVMCompilationOutput() {}
+	FVectorVMCompilationOutput(): NumOps(0){}
 
 	TArray<uint8> ByteCode;
 
@@ -25,6 +25,9 @@ struct FVectorVMCompilationOutput
 	};
 	TArray<FCalledVMFunction> CalledVMFunctionTable;
 
+	FString AssemblyAsString;
+	uint32 NumOps;
+
 	FString Errors;
 };
 
diff --git a/Engine/Source/Developer/SlateReflector/Private/Widgets/SWidgetReflector.cpp b/Engine/Source/Developer/SlateReflector/Private/Widgets/SWidgetReflector.cpp
--- a/Engine/Source/Developer/SlateReflector/Private/Widgets/SWidgetReflector.cpp
+++ b/Engine/Source/Developer/SlateReflector/Private/Widgets/SWidgetReflector.cpp
@@ -852,7 +852,7 @@ TSharedRef<SDockTab> SWidgetReflector::SpawnWidgetHierarchyTab(const FSpawnTabAr
 							.ToolTipText(LOCTEXT("VisibilityTooltip", "Visibility"))
 						]
 
-						+ SHeaderRow::Column("Focusable")
+						+ SHeaderRow::Column(SReflectorTreeWidgetItem::NAME_Focusable)
 						.DefaultLabel(LOCTEXT("Focusable", "Focusable?"))
 						.FixedWidth(125.0f)
 						.HAlignHeader(HAlign_Center)
diff --git a/Engine/Source/Developer/SlateReflector/Private/Widgets/SWidgetReflectorTreeWidgetItem.cpp b/Engine/Source/Developer/SlateReflector/Private/Widgets/SWidgetReflectorTreeWidgetItem.cpp
--- a/Engine/Source/Developer/SlateReflector/Private/Widgets/SWidgetReflectorTreeWidgetItem.cpp
+++ b/Engine/Source/Developer/SlateReflector/Private/Widgets/SWidgetReflectorTreeWidgetItem.cpp
@@ -5,7 +5,6 @@
 #include "Widgets/Colors/SColorBlock.h"
 #include "Widgets/Input/SHyperlink.h"
 #include "HAL/PlatformApplicationMisc.h"
-#include "EditorStyleSet.h"
 #include "SCheckBox.h"
 
 
@@ -15,6 +14,7 @@
 FName SReflectorTreeWidgetItem::NAME_WidgetName(TEXT("WidgetName"));
 FName SReflectorTreeWidgetItem::NAME_WidgetInfo(TEXT("WidgetInfo"));
 FName SReflectorTreeWidgetItem::NAME_Visibility(TEXT("Visibility"));
+FName SReflectorTreeWidgetItem::NAME_Focusable(TEXT("Focusable"));
 FName SReflectorTreeWidgetItem::NAME_Clipping(TEXT("Clipping"));
 FName SReflectorTreeWidgetItem::NAME_ForegroundColor(TEXT("ForegroundColor"));
 FName SReflectorTreeWidgetItem::NAME_Address(TEXT("Address"));
@@ -67,7 +67,7 @@ TSharedRef<SWidget> SReflectorTreeWidgetItem::GenerateWidgetForColumn(const FNam
 					.Justification(ETextJustify::Center)
 			];
 	}
-	else if (ColumnName == "Focusable")
+	else if (ColumnName == NAME_Focusable)
 	{
 		return SNew(SBox)
 			.HAlign(HAlign_Center)
diff --git a/Engine/Source/Developer/SlateReflector/Private/Widgets/SWidgetReflectorTreeWidgetItem.h b/Engine/Source/Developer/SlateReflector/Private/Widgets/SWidgetReflectorTreeWidgetItem.h
--- a/Engine/Source/Developer/SlateReflector/Private/Widgets/SWidgetReflectorTreeWidgetItem.h
+++ b/Engine/Source/Developer/SlateReflector/Private/Widgets/SWidgetReflectorTreeWidgetItem.h
@@ -22,6 +22,7 @@ public:
 	static FName NAME_WidgetName;
 	static FName NAME_WidgetInfo;
 	static FName NAME_Visibility;
+	static FName NAME_Focusable;
 	static FName NAME_Clipping;
 	static FName NAME_ForegroundColor;
 	static FName NAME_Address;
diff --git a/Engine/Source/Editor/AnimGraph/Private/EditModes/TwoBoneIKEditMode.cpp b/Engine/Source/Editor/AnimGraph/Private/EditModes/TwoBoneIKEditMode.cpp
--- a/Engine/Source/Editor/AnimGraph/Private/EditModes/TwoBoneIKEditMode.cpp
+++ b/Engine/Source/Editor/AnimGraph/Private/EditModes/TwoBoneIKEditMode.cpp
@@ -111,7 +111,16 @@ FVector FTwoBoneIKEditMode::GetWidgetLocation(BoneSelectModeType InBoneSelectMod
 	}
 
 	USkeletalMeshComponent* SkelComp = GetAnimPreviewScene().GetPreviewMeshComponent();
-	return ConvertWidgetLocation(SkelComp, TwoBoneIKRuntimeNode->ForwardedPose, Target, Location, Space);
+
+	// Make sure Node has had a chance to evaluate and cache a pose.
+	if (TwoBoneIKRuntimeNode->ForwardedPose.GetPose().GetNumBones() > 0)
+	{
+		return ConvertWidgetLocation(SkelComp, TwoBoneIKRuntimeNode->ForwardedPose, Target, Location, Space);
+	}
+	else
+	{
+		return SkelComp->GetComponentTransform().GetLocation();
+	}
 }
 
 FVector FTwoBoneIKEditMode::GetWidgetLocation() const
diff --git a/Engine/Source/Editor/AnimationEditor/Private/AnimationEditor.cpp b/Engine/Source/Editor/AnimationEditor/Private/AnimationEditor.cpp
--- a/Engine/Source/Editor/AnimationEditor/Private/AnimationEditor.cpp
+++ b/Engine/Source/Editor/AnimationEditor/Private/AnimationEditor.cpp
@@ -960,16 +960,17 @@ void FAnimationEditor::CopyCurveToSoundWave(const FAssetData& SoundWaveAssetData
 	}
 
 	// If no internal table, create one now
-	if (!SoundWave->InternalCurves)
+	if (!SoundWave->GetInternalCurveData())
 	{
 		static const FName InternalCurveTableName("InternalCurveTable");
-		SoundWave->Curves = NewObject<UCurveTable>(SoundWave, InternalCurveTableName);
-		SoundWave->Curves->ClearFlags(RF_Public);
-		SoundWave->Curves->SetFlags(SoundWave->Curves->GetFlags() | RF_Standalone | RF_Transactional);
-		SoundWave->InternalCurves = SoundWave->Curves;
+		UCurveTable* NewCurves = NewObject<UCurveTable>(SoundWave, InternalCurveTableName);
+		NewCurves->ClearFlags(RF_Public);
+		NewCurves->SetFlags(NewCurves->GetFlags() | RF_Standalone | RF_Transactional);
+		SoundWave->SetCurveData(NewCurves);
+		SoundWave->SetInternalCurveData(NewCurves);
 	}
 
-	UCurveTable* CurveTable = SoundWave->InternalCurves;
+	UCurveTable* CurveTable = SoundWave->GetInternalCurveData();
 
 	// iterate over curves in anim data
 	const int32 NumCurves = Sequence->RawCurveData.FloatCurves.Num();
diff --git a/Engine/Source/Editor/BehaviorTreeEditor/Private/BehaviorTreeGraph.cpp b/Engine/Source/Editor/BehaviorTreeEditor/Private/BehaviorTreeGraph.cpp
--- a/Engine/Source/Editor/BehaviorTreeEditor/Private/BehaviorTreeGraph.cpp
+++ b/Engine/Source/Editor/BehaviorTreeEditor/Private/BehaviorTreeGraph.cpp
@@ -60,6 +60,7 @@ void UBehaviorTreeGraph::UpdateBlackboardChange()
 			UBTNode* MyNodeInstance = Cast<UBTNode>(MyNode->NodeInstance);
 			if (MyNodeInstance)
 			{
+				UBehaviorTreeTypes::SetBTLoggingContext(MyNodeInstance);
 				MyNodeInstance->InitializeFromAsset(*BTAsset);
 			}
 
@@ -68,6 +69,7 @@ void UBehaviorTreeGraph::UpdateBlackboardChange()
 				UBTNode* DecoratorNodeInstance = MyNode->Decorators[iDecorator] ? Cast<UBTNode>(MyNode->Decorators[iDecorator]->NodeInstance) : NULL;
 				if (DecoratorNodeInstance)
 				{
+					UBehaviorTreeTypes::SetBTLoggingContext(DecoratorNodeInstance);
 					DecoratorNodeInstance->InitializeFromAsset(*BTAsset);
 				}
 
@@ -83,9 +85,12 @@ void UBehaviorTreeGraph::UpdateBlackboardChange()
 				UBTNode* ServiceNodeInstance = MyNode->Services[iService] ? Cast<UBTNode>(MyNode->Services[iService]->NodeInstance) : NULL;
 				if (ServiceNodeInstance)
 				{
+					UBehaviorTreeTypes::SetBTLoggingContext(ServiceNodeInstance);
 					ServiceNodeInstance->InitializeFromAsset(*BTAsset);
 				}
 			}
+
+			UBehaviorTreeTypes::SetBTLoggingContext(nullptr);
 		}
 	}
 }
@@ -173,6 +178,8 @@ void UBehaviorTreeGraph::UpdateAsset(int32 UpdateFlags)
 			}
 		}
 	}
+
+	UpdateBlackboardChange();
 }
 
 void UBehaviorTreeGraph::OnCreated()
@@ -191,7 +198,6 @@ void UBehaviorTreeGraph::OnLoaded()
 void UBehaviorTreeGraph::Initialize()
 {
 	Super::Initialize();
-	UpdateBlackboardChange();
 	UpdateInjectedNodes();
 }
 
diff --git a/Engine/Source/Editor/DetailCustomizations/Private/SoundWaveDetails.cpp b/Engine/Source/Editor/DetailCustomizations/Private/SoundWaveDetails.cpp
--- a/Engine/Source/Editor/DetailCustomizations/Private/SoundWaveDetails.cpp
+++ b/Engine/Source/Editor/DetailCustomizations/Private/SoundWaveDetails.cpp
@@ -34,7 +34,7 @@ void FSoundWaveDetails::CustomizeCurveDetails(IDetailLayoutBuilder& DetailBuilde
 	{
 		USoundWave* SoundWave = CastChecked<USoundWave>(Objects[0].Get());
 
-		TSharedRef<IPropertyHandle> CurvePropertyHandle = DetailBuilder.GetProperty(GET_MEMBER_NAME_CHECKED(USoundWave, Curves));
+		TSharedRef<IPropertyHandle> CurvePropertyHandle = DetailBuilder.GetProperty(USoundWave::GetCurvePropertyName());
 		if(CurvePropertyHandle->IsValidHandle())
 		{
 			IDetailPropertyRow& CurvePropertyRow = DetailBuilder.EditCategory(TEXT("Curves")).AddProperty(CurvePropertyHandle);
@@ -105,7 +105,7 @@ EVisibility FSoundWaveDetails::GetMakeInternalCurvesVisibility(USoundWave* Sound
 
 EVisibility FSoundWaveDetails::GetUseInternalCurvesVisibility(USoundWave* SoundWave, TSharedRef<IPropertyHandle> CurvePropertyHandle) const
 {
-	UCurveTable* InternalCurveTable = SoundWave->InternalCurves;
+	UCurveTable* InternalCurveTable = SoundWave->GetInternalCurveData();
 
 	UObject* CurrentCurveTable = nullptr;
 	if (CurvePropertyHandle->GetValue(CurrentCurveTable) == FPropertyAccess::Success)
@@ -121,15 +121,17 @@ EVisibility FSoundWaveDetails::GetUseInternalCurvesVisibility(USoundWave* SoundW
 
 FReply FSoundWaveDetails::HandleMakeInternalCurves(USoundWave* SoundWave)
 {
-	if (SoundWave->Curves != nullptr)
+	UCurveTable* ExistingCurves = SoundWave->GetCurveData();
+	if (ExistingCurves != nullptr)
 	{
 		FScopedTransaction Transaction(LOCTEXT("MakeInternalCurve", "Copy Curve to Internal"));
 		SoundWave->Modify();
 
-		SoundWave->Curves = DuplicateObject<UCurveTable>(SoundWave->Curves, SoundWave, InternalCurveTableName);
-		SoundWave->Curves->ClearFlags(RF_Public);
-		SoundWave->Curves->SetFlags(SoundWave->Curves->GetFlags() | RF_Standalone | RF_Transactional);
-		SoundWave->InternalCurves = SoundWave->Curves;
+		UCurveTable* DuplicatedCurves = DuplicateObject<UCurveTable>(ExistingCurves, SoundWave, InternalCurveTableName);
+		DuplicatedCurves->ClearFlags(RF_Public);
+		DuplicatedCurves->SetFlags(ExistingCurves->GetFlags() | RF_Standalone | RF_Transactional);
+		SoundWave->SetCurveData(DuplicatedCurves);
+		SoundWave->SetInternalCurveData(DuplicatedCurves);
 	}
 
 	return FReply::Handled();
@@ -137,7 +139,7 @@ FReply FSoundWaveDetails::HandleMakeInternalCurves(USoundWave* SoundWave)
 
 FReply FSoundWaveDetails::HandleUseInternalCurves(USoundWave* SoundWave, TSharedRef<IPropertyHandle> CurvePropertyHandle)
 {
-	UCurveTable* InternalCurveTable = SoundWave->InternalCurves;
+	UCurveTable* InternalCurveTable = SoundWave->GetInternalCurveData();
 	if (InternalCurveTable != nullptr)
 	{
 		CurvePropertyHandle->SetValue(InternalCurveTable);
diff --git a/Engine/Source/Editor/Kismet/Private/BlueprintEditor.cpp b/Engine/Source/Editor/Kismet/Private/BlueprintEditor.cpp
--- a/Engine/Source/Editor/Kismet/Private/BlueprintEditor.cpp
+++ b/Engine/Source/Editor/Kismet/Private/BlueprintEditor.cpp
@@ -4593,7 +4593,7 @@ bool FBlueprintEditor::CanCollapseSelectionToFunction(TSet<class UEdGraphNode*>&
 		{
 			for (UEdGraphPin* NodePin : Node->Pins)
 			{
-				if(NodePin->PinType.PinCategory == UEdGraphSchema_K2::PC_Exec)
+				if (NodePin->PinType.PinCategory == K2Schema->PC_Exec)
 				{
 					if (NodePin->LinkedTo.Num() == 0 && NodePin->Direction == EGPD_Input)
 					{
diff --git a/Engine/Source/Editor/PhysicsAssetEditor/Private/PhysicsAssetEditorSharedData.cpp b/Engine/Source/Editor/PhysicsAssetEditor/Private/PhysicsAssetEditorSharedData.cpp
--- a/Engine/Source/Editor/PhysicsAssetEditor/Private/PhysicsAssetEditorSharedData.cpp
+++ b/Engine/Source/Editor/PhysicsAssetEditor/Private/PhysicsAssetEditorSharedData.cpp
@@ -70,6 +70,7 @@ void FPhysicsAssetEditorSharedData::Initialize(const TSharedRef<IPersonaPreviewS
 
 	EditorSkelComp = nullptr;
 	PhysicalAnimationComponent = nullptr;
+	FStringAssetReference PreviewMeshStringRef = PhysicsAsset->PreviewSkeletalMesh.ToSoftObjectPath();
 
 	// Look for body setups with no shapes (how does this happen?).
 	// If we find one- just bang on a default box.
diff --git a/Engine/Source/Editor/PropertyEditor/Private/CustomChildBuilder.cpp b/Engine/Source/Editor/PropertyEditor/Private/CustomChildBuilder.cpp
--- a/Engine/Source/Editor/PropertyEditor/Private/CustomChildBuilder.cpp
+++ b/Engine/Source/Editor/PropertyEditor/Private/CustomChildBuilder.cpp
@@ -95,13 +95,13 @@ TArray<TSharedPtr<IPropertyHandle>> FCustomChildrenBuilder::AddAllExternalStruct
 	return ParentCategory.Pin()->AddAllExternalStructureProperties(ChildStructure);
 }
 
-IDetailPropertyRow* FCustomChildrenBuilder::AddExternalObjectProperty(const TArray<UObject*>& Objects, FName PropertyName, FName UniqueIdName)
+IDetailPropertyRow* FCustomChildrenBuilder::AddExternalObjectProperty(const TArray<UObject*>& Objects, FName PropertyName, FName UniqueIdName, TOptional<bool> bAllowChildrenOverride, TOptional<bool> bCreateCategoryNodesOverride)
 {
 	FDetailLayoutCustomization NewCustomization;
 
 	TSharedRef<FDetailCategoryImpl> ParentCategoryRef = ParentCategory.Pin().ToSharedRef();
 
-	FDetailPropertyRow::MakeExternalPropertyRowCustomization(Objects, PropertyName, ParentCategoryRef, NewCustomization);
+	FDetailPropertyRow::MakeExternalPropertyRowCustomization(Objects, PropertyName, ParentCategoryRef, NewCustomization, bAllowChildrenOverride, bCreateCategoryNodesOverride);
 
 	TSharedPtr<FDetailPropertyRow> NewRow = NewCustomization.PropertyRow;
 
diff --git a/Engine/Source/Editor/PropertyEditor/Private/CustomChildBuilder.h b/Engine/Source/Editor/PropertyEditor/Private/CustomChildBuilder.h
--- a/Engine/Source/Editor/PropertyEditor/Private/CustomChildBuilder.h
+++ b/Engine/Source/Editor/PropertyEditor/Private/CustomChildBuilder.h
@@ -25,7 +25,7 @@ public:
 	virtual FDetailWidgetRow& AddCustomRow(const FText& SearchString) override;
 	virtual IDetailPropertyRow& AddProperty(TSharedRef<IPropertyHandle> PropertyHandle) override;
 	virtual IDetailPropertyRow* AddExternalObjects(const TArray<UObject*>& Objects, FName UniqueIdName = NAME_None) override;
-	virtual IDetailPropertyRow* AddExternalObjectProperty(const TArray<UObject*>& Objects, FName PropertyName, FName UniqueIdName = NAME_None) override;
+	virtual IDetailPropertyRow* AddExternalObjectProperty(const TArray<UObject*>& Objects, FName PropertyName, FName UniqueIdName = NAME_None, TOptional<bool> bAllowChildrenOverride = TOptional<bool>(), TOptional<bool> bCreateCategoryNodesOverride = TOptional<bool>()) override;
 	virtual IDetailPropertyRow* AddExternalStructure(TSharedRef<FStructOnScope> ChildStructure, FName UniqueIdName = NAME_None) override;
 	virtual IDetailPropertyRow* AddExternalStructureProperty(TSharedRef<FStructOnScope> ChildStructure, FName PropertyName, FName UniqueIdName = NAME_None) override;
 	virtual TArray<TSharedPtr<IPropertyHandle>> AddAllExternalStructureProperties(TSharedRef<FStructOnScope> ChildStructure) override;
diff --git a/Engine/Source/Editor/PropertyEditor/Private/DetailCategoryBuilderImpl.cpp b/Engine/Source/Editor/PropertyEditor/Private/DetailCategoryBuilderImpl.cpp
--- a/Engine/Source/Editor/PropertyEditor/Private/DetailCategoryBuilderImpl.cpp
+++ b/Engine/Source/Editor/PropertyEditor/Private/DetailCategoryBuilderImpl.cpp
@@ -86,8 +86,12 @@ FDetailCategoryImpl::FDetailCategoryImpl(FName InCategoryName, TSharedRef<FDetai
 	, bHasVisibleDetails(true)
 	, bIsCategoryVisible(true)
 	, bFavoriteCategory(false)
+	, bShowOnlyChildren(false)
 {
 	const UStruct* BaseStruct = InDetailLayout->GetRootNode()->GetBaseStructure();
+
+	bShowOnlyChildren = InDetailLayout->IsLayoutForExternalRoot() && !InDetailLayout->GetRootNode()->HasNodeFlags(EPropertyNodeFlags::ShowCategories);
+
 	// Use the base class name if there is one otherwise this is a generic category not specific to a class
 	FName BaseStructName = BaseStruct ? BaseStruct->GetFName() : FName("Generic");
 
diff --git a/Engine/Source/Editor/PropertyEditor/Private/DetailCategoryBuilderImpl.h b/Engine/Source/Editor/PropertyEditor/Private/DetailCategoryBuilderImpl.h
--- a/Engine/Source/Editor/PropertyEditor/Private/DetailCategoryBuilderImpl.h
+++ b/Engine/Source/Editor/PropertyEditor/Private/DetailCategoryBuilderImpl.h
@@ -185,7 +185,7 @@ public:
 	virtual ENodeVisibility GetVisibility() const override;
 	virtual void FilterNode(const FDetailFilter& DetailFilter) override;
 	virtual void Tick(float DeltaTime) override {}
-	virtual bool ShouldShowOnlyChildren() const override { return false; }
+	virtual bool ShouldShowOnlyChildren() const override { return bShowOnlyChildren; }
 	virtual FName GetNodeName() const override { return GetCategoryName(); }
 
 	/**
@@ -419,4 +419,5 @@ private:
 	bool bIsCategoryVisible : 1;
 	/*true if the category is the special favorite category, all property in the layout will be display when we generate the roottree */
 	bool bFavoriteCategory : 1;
+	bool bShowOnlyChildren : 1;
 };
diff --git a/Engine/Source/Editor/PropertyEditor/Private/DetailLayoutBuilderImpl.cpp b/Engine/Source/Editor/PropertyEditor/Private/DetailLayoutBuilderImpl.cpp
--- a/Engine/Source/Editor/PropertyEditor/Private/DetailLayoutBuilderImpl.cpp
+++ b/Engine/Source/Editor/PropertyEditor/Private/DetailLayoutBuilderImpl.cpp
@@ -9,14 +9,14 @@
 #include "DetailMultiTopLevelObjectRootNode.h"
 #include "ObjectEditorUtils.h"
 
-FDetailLayoutBuilderImpl::FDetailLayoutBuilderImpl(TSharedPtr<FComplexPropertyNode>& InRootNode, FClassToPropertyMap& InPropertyMap, const TSharedRef< class IPropertyUtilities >& InPropertyUtilities, const TSharedPtr< IDetailsViewPrivate >& InDetailsView)
+FDetailLayoutBuilderImpl::FDetailLayoutBuilderImpl(TSharedPtr<FComplexPropertyNode>& InRootNode, FClassToPropertyMap& InPropertyMap, const TSharedRef< class IPropertyUtilities >& InPropertyUtilities, const TSharedPtr< IDetailsViewPrivate >& InDetailsView, bool bIsExternal)
 	: RootNode( InRootNode )
 	, PropertyMap( InPropertyMap )
 	, PropertyDetailsUtilities( InPropertyUtilities )
 	, DetailsView( InDetailsView.Get() )
 	, CurrentCustomizationClass( nullptr )
 {
-	bLayoutForExternalRoot = ExternalRootPropertyNodes.Contains(InRootNode.ToSharedRef());
+	bLayoutForExternalRoot = bIsExternal;
 }
 
 IDetailCategoryBuilder& FDetailLayoutBuilderImpl::EditCategory( FName CategoryName, const FText& NewLocalizedDisplayName, ECategoryPriority::Type CategoryType )
diff --git a/Engine/Source/Editor/PropertyEditor/Private/DetailLayoutBuilderImpl.h b/Engine/Source/Editor/PropertyEditor/Private/DetailLayoutBuilderImpl.h
--- a/Engine/Source/Editor/PropertyEditor/Private/DetailLayoutBuilderImpl.h
+++ b/Engine/Source/Editor/PropertyEditor/Private/DetailLayoutBuilderImpl.h
@@ -20,7 +20,8 @@ public:
 		TSharedPtr<FComplexPropertyNode>& InRootNode,
 		FClassToPropertyMap& InPropertyMap,
 		const TSharedRef<IPropertyUtilities>& InPropertyUtilities,
-		const TSharedPtr<IDetailsViewPrivate>& InDetailsView);
+		const TSharedPtr<IDetailsViewPrivate>& InDetailsView,
+		bool bIsExternal);
 
 	/** IDetailLayoutBuilder Interface */
 	virtual const IDetailsView* GetDetailsView() const override;
