Fix for VM Crash

#codereivew Frank.Fella, Shaun.Kime, Olaf.Piesche

diff --git a/Engine/Plugins/FX/Niagara/Source/NiagaraCore/Private/NiagaraCustomVersion.cpp b/Engine/Plugins/FX/Niagara/Source/NiagaraCore/Private/NiagaraCustomVersion.cpp
--- a/Engine/Plugins/FX/Niagara/Source/NiagaraCore/Private/NiagaraCustomVersion.cpp
+++ b/Engine/Plugins/FX/Niagara/Source/NiagaraCore/Private/NiagaraCustomVersion.cpp
@@ -8,4 +8,4 @@ const FGuid FNiagaraCustomVersion::GUID(0xFCF57AFA, 0x50764283, 0xB9A9E658, 0xFF
 // Register the custom version with core
 FCustomVersionRegistration GRegisterNiagaraCustomVersion(FNiagaraCustomVersion::GUID, FNiagaraCustomVersion::LatestVersion, TEXT("NiagaraVer"));
 
-const FGuid FNiagaraCustomVersion::LatestScriptCompileVersion(0x73061024, 0xFC2646FF, 0x82F694A8, 0xF619C0CE);
\ No newline at end of file
+const FGuid FNiagaraCustomVersion::LatestScriptCompileVersion(0x8A64E49F, 0xF987452E, 0x9004F80D, 0xB2EBDE64);
\ No newline at end of file
diff --git a/Engine/Source/Developer/ShaderFormatVectorVM/Private/ir_vm_gen_bytecode_visitor.cpp b/Engine/Source/Developer/ShaderFormatVectorVM/Private/ir_vm_gen_bytecode_visitor.cpp
--- a/Engine/Source/Developer/ShaderFormatVectorVM/Private/ir_vm_gen_bytecode_visitor.cpp
+++ b/Engine/Source/Developer/ShaderFormatVectorVM/Private/ir_vm_gen_bytecode_visitor.cpp
@@ -2004,32 +2004,28 @@ class ir_gen_vvm_visitor : public ir_hierarchical_visitor
 #if VM_VERBOSE_LOGGING
 		UE_LOG(LogVVMBackend, Log, TEXT("\n-------------------------------\nTemporary Allocations\n------------------------------\n"));
 #endif
-		TArray<int32> ops_to_strip;
+		TArray<op_base*> stripped_ops;
+		stripped_ops.Reserve(ordered_ops.Num());
 		for (int32 i = 0; i < ordered_ops.Num(); ++i)
 		{
-			if (!ordered_ops[i]->finalize(parse_state, CompilationOutput, num_temp_registers, registers, i))
+			if (ordered_ops[i]->finalize(parse_state, CompilationOutput, num_temp_registers, registers, i))
 			{
-				ops_to_strip.Add(i);
+				stripped_ops.Add(ordered_ops[i]);
 			}
 		}
-
-		for (int32 op_idx : ops_to_strip)
-		{
-			ordered_ops.RemoveAt(op_idx);
-		}
-
+		
 		//Final error checking.
-		for (int32 i = 0; i < ordered_ops.Num(); ++i)
+		for (int32 i = 0; i < stripped_ops.Num(); ++i)
 		{
-			ordered_ops[i]->validate(parse_state);
+			stripped_ops[i]->validate(parse_state);
 		}
 
 		if (!parse_state->error)
 		{
 			//Now emit the bytecode
-			for (int32 op_idx = 0; op_idx < ordered_ops.Num(); ++op_idx)
+			for (int32 op_idx = 0; op_idx < stripped_ops.Num(); ++op_idx)
 			{
-				ordered_ops[op_idx]->add_to_bytecode(CompilationOutput.ByteCode);
+				stripped_ops[op_idx]->add_to_bytecode(CompilationOutput.ByteCode);
 			}
 
 			emit_byte((uint8)EVectorVMOp::done, CompilationOutput.ByteCode);
diff --git a/Engine/Source/Developer/ShaderFormatVectorVM/Private/ir_vm_scalarize_visitor.cpp b/Engine/Source/Developer/ShaderFormatVectorVM/Private/ir_vm_scalarize_visitor.cpp
--- a/Engine/Source/Developer/ShaderFormatVectorVM/Private/ir_vm_scalarize_visitor.cpp
+++ b/Engine/Source/Developer/ShaderFormatVectorVM/Private/ir_vm_scalarize_visitor.cpp
@@ -398,6 +398,11 @@ class ir_scalarize_visitor2 : public ir_hierarchical_visitor
 				if (comp_assign)
 				{
 					assign->insert_before(comp_assign);
+					//Ensure that we have a valid write mask if we're splitting down a structure assignment.
+					if (comp_assign->lhs->type->is_scalar())
+					{
+						comp_assign->write_mask = 1;
+					}
 					comp_assign = NULL;
 				}
 
@@ -428,6 +433,11 @@ class ir_scalarize_visitor2 : public ir_hierarchical_visitor
 		if (comp_assign)
 		{
 			assign->replace_with(comp_assign);
+			//Ensure that we have a valid write mask if we're splitting down a structure assignment.
+			if (comp_assign->lhs->type->is_scalar())
+			{
+				comp_assign->write_mask = 1;
+			}
 		}
 
 		return visit_continue_with_parent;
