Disabling the vectorvm fix in main. It will propagate up through dev-rendering.

diff --git a/Engine/Source/Developer/ShaderFormatVectorVM/Private/VectorVMShaderCompiler.cpp b/Engine/Source/Developer/ShaderFormatVectorVM/Private/VectorVMShaderCompiler.cpp
--- a/Engine/Source/Developer/ShaderFormatVectorVM/Private/VectorVMShaderCompiler.cpp
+++ b/Engine/Source/Developer/ShaderFormatVectorVM/Private/VectorVMShaderCompiler.cpp
@@ -104,7 +104,7 @@ bool CompileShader_VectorVM(const FShaderCompilerInput& Input, FShaderCompilerOu
 
 	if (bSkipBackendOptimizations)
 	{
-		CCFlags |= HLSLCC_DisableBackendOptimizations;
+		//CCFlags |= HLSLCC_DisableBackendOptimizations;
 	}
 
 	//TODO: Do this later when we implement the rest of the shader plumbing stuff.
