Handling edge case in the non-expression visitor.
#tests n/a

diff --git a/Engine/Source/Developer/ShaderFormatVectorVM/Private/ir_vm_propagate_non_expressions_visitor.cpp b/Engine/Source/Developer/ShaderFormatVectorVM/Private/ir_vm_propagate_non_expressions_visitor.cpp
--- a/Engine/Source/Developer/ShaderFormatVectorVM/Private/ir_vm_propagate_non_expressions_visitor.cpp
+++ b/Engine/Source/Developer/ShaderFormatVectorVM/Private/ir_vm_propagate_non_expressions_visitor.cpp
@@ -105,11 +105,11 @@ public:
 			{
 				//Is there a previous non_expr assignment after any containing expressions?
 				//If so, copy that in place of this rvalue.
-				if (varinfo->latest_expr_assign[search_comp] < varinfo->latest_non_expr_assign[search_comp])
+				if (varinfo->latest_expr_assign[search_comp] < varinfo->latest_non_expr_assign[search_comp] )
 				{
 					ir_assignment* assign = assignments[varinfo->latest_non_expr_assign[search_comp]];
 					check(assign->rhs->as_expression() == nullptr);
-					check(assign->rhs->as_swizzle() || assign->rhs->as_dereference_variable() || assign->rhs->as_constant());
+					check(assign->rhs->as_swizzle() || assign->rhs->as_dereference_variable() || assign->rhs->as_constant() || assign->rhs->as_dereference_record());
 					check(assign->rhs->type->is_scalar());//All assignments must be scalar at this point!
 					ir_rvalue* new_rval = assign->rhs->clone(parse_state, nullptr);
 					(*rvalue) = new_rval;
