Now properly supporting DI function calls in the VM having unused output parameters.

diff --git a/Engine/Plugins/FX/Niagara/Source/NiagaraCore/Private/NiagaraCustomVersion.cpp b/Engine/Plugins/FX/Niagara/Source/NiagaraCore/Private/NiagaraCustomVersion.cpp
--- a/Engine/Plugins/FX/Niagara/Source/NiagaraCore/Private/NiagaraCustomVersion.cpp
+++ b/Engine/Plugins/FX/Niagara/Source/NiagaraCore/Private/NiagaraCustomVersion.cpp
@@ -8,4 +8,4 @@ const FGuid FNiagaraCustomVersion::GUID(0xFCF57AFA, 0x50764283, 0xB9A9E658, 0xFF
 // Register the custom version with core
 FCustomVersionRegistration GRegisterNiagaraCustomVersion(FNiagaraCustomVersion::GUID, FNiagaraCustomVersion::LatestVersion, TEXT("NiagaraVer"));
 
-const FGuid FNiagaraCustomVersion::LatestScriptCompileVersion(0x0D9851F7, 0x8A2C4C60, 0xB17F4DEB, 0x8435B868);
+const FGuid FNiagaraCustomVersion::LatestScriptCompileVersion(0x9BA7BE7D, 0x57FD4A80, 0xB2DD5740, 0xB9FCCC78);
\ No newline at end of file
diff --git a/Engine/Source/Developer/ShaderFormatVectorVM/Private/ir_vm_gen_bytecode_visitor.cpp b/Engine/Source/Developer/ShaderFormatVectorVM/Private/ir_vm_gen_bytecode_visitor.cpp
--- a/Engine/Source/Developer/ShaderFormatVectorVM/Private/ir_vm_gen_bytecode_visitor.cpp
+++ b/Engine/Source/Developer/ShaderFormatVectorVM/Private/ir_vm_gen_bytecode_visitor.cpp
@@ -596,11 +596,11 @@ struct op_base
 	{
 		check(component);
 
-		if (component->last_read <= op_idx)
+		if (component->last_read <= op_idx || component->last_read == INDEX_NONE)
 		{
-			//This component is not used so we can avoid using a temporary for it in some cases.
+			//It's possible that some components can have invalid indices here. We let the calling operations deal with this as they see fit.
 			check(component->offset == INDEX_NONE);
-			return true;
+			return false;
 		}
 
 		if (ir_variable* var = component->owner->ir->as_variable())
@@ -1119,9 +1119,16 @@ struct op_external_func : public op_base
 
 			ComponentStr.Empty(64);
 			check(output->is_scalar());
-			ComponentStr = LexToString(output->offset);
 
-			Str += FString::Printf(TEXT("%s[%s]"), get_location_string(output->owner->location), *ComponentStr);
+			if (output->offset == INDEX_NONE)
+			{
+				Str += TEXT("Unused");
+			}
+			else
+			{
+				ComponentStr = LexToString(output->offset);
+				Str += FString::Printf(TEXT("%s[%s]"), get_location_string(output->owner->location), *ComponentStr);
+			}
 		}
 
 		return Str += TEXT(");\n");
@@ -1133,7 +1140,7 @@ struct op_external_func : public op_base
 		bool func_is_used = false;
 		for (variable_info_node* output : outputs)
 		{
-			if (output->last_read > op_idx || output->last_read != INDEX_NONE)
+			if (output->last_read > op_idx && output->last_read != INDEX_NONE)
 			{
 				func_is_used = true;
 				break;
@@ -1179,7 +1186,8 @@ struct op_external_func : public op_base
 		finalize_temp_register_ctx* ctx = (finalize_temp_register_ctx*)user_ptr;
 		check(node->is_scalar());
 
-		ctx->success &= finalize_component_temporary_allocation(ctx->parse_state, node, ctx->num_temp_registers, ctx->registers, ctx->op_idx);
+		//Some outputs for an external function can have invalid locations. If the thing has *any* valid outputs then it's considered a success.
+		ctx->success |= finalize_component_temporary_allocation(ctx->parse_state, node, ctx->num_temp_registers, ctx->registers, ctx->op_idx);
 	}
 
 	virtual bool finalize(_mesa_glsl_parse_state* parse_state, FVectorVMCompilationOutput&CompilationOutput, const unsigned num_temp_registers, unsigned *registers, unsigned op_idx)
@@ -1214,7 +1222,7 @@ struct op_external_func : public op_base
 		}
 		if (!ctx.success)
 		{
-			_mesa_glsl_error(parse_state, "Failed to allocate temporary registers for external funciton.");
+			_mesa_glsl_error(parse_state, "Failed to allocate temporary registers for external function.");
 		}
 		return ctx.success;
 	}
@@ -1230,10 +1238,12 @@ struct op_external_func : public op_base
 		{
 			validate_component_offset(parse_state, input, op_idx);
 		}
-		for (variable_info_node* output : outputs)
-		{
-			validate_component_offset(parse_state, output, op_idx);
-		}
+
+		//Output offsets are allowed to be INDEX_NONE if they are unused.
+// 		for (variable_info_node* output : outputs)
+// 		{
+// 			validate_component_offset(parse_state, output, op_idx);
+// 		}
 	}
 };
 
@@ -2199,4 +2209,4 @@ void vm_gen_bytecode(exec_list *ir, _mesa_glsl_parse_state *state, FVectorVMComp
 	BuildExpressionMap();
 
 	ir_gen_vvm_visitor::run(ir, state, InCompOutput);
-}
+}
\ No newline at end of file
