Fix for matrix params in external function calls. Or any function call for that matter.

diff --git a/Engine/Source/Developer/ShaderFormatVectorVM/Private/ir_vm_propagate_non_expressions_visitor.cpp b/Engine/Source/Developer/ShaderFormatVectorVM/Private/ir_vm_propagate_non_expressions_visitor.cpp
--- a/Engine/Source/Developer/ShaderFormatVectorVM/Private/ir_vm_propagate_non_expressions_visitor.cpp
+++ b/Engine/Source/Developer/ShaderFormatVectorVM/Private/ir_vm_propagate_non_expressions_visitor.cpp
@@ -226,6 +226,55 @@ public:
 	{
 	}
 
+	virtual ir_visitor_status visit_enter(ir_call* call)
+	{
+		bool bFoundMatrixParam = false;
+		foreach_iter(exec_list_iterator, iter, call->actual_parameters)
+		{
+			ir_rvalue* param = (ir_rvalue*)iter.get();
+			if (param->type->is_matrix())
+			{
+				bFoundMatrixParam = true;
+				break;
+			}
+		}
+
+		if (!bFoundMatrixParam)
+		{
+			return ir_rvalue_visitor::visit_enter(call);
+		}
+
+		exec_list old_params = call->actual_parameters;
+
+		call->actual_parameters.make_empty();
+
+		//Replace any matrix params with params for their vector components.
+		//This will all be full scalarized to float derefs anyway in later passes.
+		auto add_param = [&](ir_rvalue* rval)
+		{
+			if (rval->type->is_matrix())
+			{
+				MatrixVectors& mv = MatrixVectorMap.FindChecked(rval->variable_referenced());
+				call->actual_parameters.push_tail(new(parse_state) ir_dereference_variable(mv.v[0]));
+				call->actual_parameters.push_tail(new(parse_state) ir_dereference_variable(mv.v[1]));
+				call->actual_parameters.push_tail(new(parse_state) ir_dereference_variable(mv.v[2]));
+				call->actual_parameters.push_tail(new(parse_state) ir_dereference_variable(mv.v[3]));
+			}
+			else
+			{
+				call->actual_parameters.push_tail(rval);
+			}
+		};
+
+		foreach_iter(exec_list_iterator, param_iter, old_params)
+		{
+			ir_rvalue* param = (ir_rvalue*)param_iter.get();
+			add_param(param);
+		}
+
+		return visit_continue;
+	}
+
 	virtual void handle_rvalue(ir_rvalue **rvalue)
 	{
 		if (rvalue && *rvalue)
diff --git a/Engine/Source/Developer/ShaderFormatVectorVM/Private/ir_vm_scalarize_visitor.cpp b/Engine/Source/Developer/ShaderFormatVectorVM/Private/ir_vm_scalarize_visitor.cpp
--- a/Engine/Source/Developer/ShaderFormatVectorVM/Private/ir_vm_scalarize_visitor.cpp
+++ b/Engine/Source/Developer/ShaderFormatVectorVM/Private/ir_vm_scalarize_visitor.cpp
@@ -93,7 +93,7 @@ class ir_scalarize_visitor2 : public ir_hierarchical_visitor
 
 	virtual ir_visitor_status visit(ir_constant *ir)
 	{
-		check(ir->type->is_numeric());
+		check(ir->type->is_numeric() || ir->type->is_boolean());
 
 		unsigned use_component = FMath::Min(ir->type->components(), dest_component);
 
@@ -517,7 +517,7 @@ class ir_scalarize_visitor2 : public ir_hierarchical_visitor
 		}
 		else
 		{
-			check(deref->type->is_numeric());
+			check(deref->type->is_numeric() || deref->type->is_boolean());
 			//check(type->components() >= dest_component || type->is_scalar());
 			unsigned use_component = FMath::Min(type->components(), dest_component);
 
@@ -544,7 +544,7 @@ class ir_scalarize_visitor2 : public ir_hierarchical_visitor
 		{
 			if (!type->is_scalar())
 			{
-				check(var->type->is_numeric());
+				check(var->type->is_numeric() || var->type->is_boolean());
 				//check(type->components() > dest_component);
 				unsigned use_component = FMath::Min(type->components(), dest_component);
 
