- Storing compact form of databuffer layout in single array to reduce cache misses in register setup.
- Optimizing temp register layout for better cache usage.
- Moving VM context from TLS to pool.

diff --git a/Engine/Plugins/FX/Niagara/Source/Niagara/Classes/NiagaraDataSet.h b/Engine/Plugins/FX/Niagara/Source/Niagara/Classes/NiagaraDataSet.h
--- a/Engine/Plugins/FX/Niagara/Source/Niagara/Classes/NiagaraDataSet.h
+++ b/Engine/Plugins/FX/Niagara/Source/Niagara/Classes/NiagaraDataSet.h
@@ -8,6 +8,7 @@
 #include "VectorVM.h"
 #include "RenderingThread.h"
 
+typedef TArray<uint8*, TInlineAllocator<128>> FNiagaraRegisterTable;
 
 /** Helper class defining the layout and location of an FNiagaraVariable in an FNiagaraDataBuffer-> */
 struct FNiagaraVariableLayoutInfo
@@ -104,7 +105,8 @@ public:
 	void GPUCopyFrom(float* GPUReadBackFloat, int* GPUReadBackInt, int32 StartIdx, int32 NumInstances, uint32 InSrcFloatStride, uint32 InSrcIntStride);
 	void Dump(int32 StartIndex, int32 NumInstances, const FString& Label)const;
 
-	bool AppendToRegisterTable(uint8** Registers, int32& NumRegisters, int32 StartInstance);
+	FORCEINLINE void AppendToRegisterTable(FNiagaraRegisterTable& Registers, int32 StartInstance);
+	FORCEINLINE void ClearRegisterTable(FNiagaraRegisterTable& Registers);
 
 	const TArray<uint8>& GetFloatBuffer()const { return FloatData; }
 	const TArray<uint8>& GetInt32Buffer()const { return Int32Data; }
@@ -151,9 +153,12 @@ public:
 
 	void ReleaseGPUInstanceCount(FNiagaraGPUInstanceCountManager& GPUInstanceCountManager);
 
+	FORCEINLINE const TArray<uint8*>& GetRegisterTable()const { return RegisterTable; }
+
+	void BuildRegisterTable();
 private:
 
-	void CheckUsage(bool bReadOnly)const;
+	FORCEINLINE void CheckUsage(bool bReadOnly)const;
 
 	FORCEINLINE int32 GetSafeComponentBufferSize(int32 RequiredSize) const
 	{
@@ -196,6 +201,9 @@ private:
 	uint32 FloatStride;
 	/** Stride between components in the int32 buffer. */
 	uint32 Int32Stride;
+
+	/** Table containing current base locations for all registers in this dataset. */
+	TArray<uint8*> RegisterTable;//TODO: Should make inline? Feels like a useful size to keep local would be too big.
 };
 
 //////////////////////////////////////////////////////////////////////////
@@ -233,8 +241,6 @@ public:
 	/** Returns size in bytes for all data buffers currently allocated by this dataset. */
 	uint32 GetSizeBytes()const;
 
-	void ClearRegisterTable(uint8** Registers, int32& NumRegisters);
-
 	FORCEINLINE bool IsInitialized()const { return bFinalized; }
 	FORCEINLINE ENiagaraSimTarget GetSimTarget()const { return SimTarget; }
 	FORCEINLINE FNiagaraDataSetID GetID()const { return ID; }
@@ -1497,3 +1503,46 @@ private:
 	NiagaraEmitterInstanceBatcher* Batcher = nullptr;
 	uint32				NumInstances = 0;
 };
+
+
+//////////////////////////////////////////////////////////////////////////
+
+FORCEINLINE void FNiagaraDataBuffer::CheckUsage(bool bReadOnly)const
+{
+	checkSlow(Owner);
+
+	//We can read on the RT but any modifications must be GT (or GT Task).
+	//For GPU sims we must be on the RT.
+	checkSlow((Owner->SimTarget == ENiagaraSimTarget::CPUSim  && (IsInGameThread() || (bReadOnly || !IsInRenderingThread()))) ||
+		(Owner->SimTarget == ENiagaraSimTarget::GPUComputeSim) && IsInRenderingThread());
+}
+
+FORCEINLINE void FNiagaraDataBuffer::AppendToRegisterTable(FNiagaraRegisterTable& Registers, int32 StartInstance)
+{
+	checkSlow(Owner && Owner->IsInitialized());
+	checkSlow(Owner->GetSimTarget() == ENiagaraSimTarget::CPUSim);
+	CheckUsage(true);
+
+	int32 RegisterCount = RegisterTable.Num();
+	int32 NumRegisters = Registers.Num();
+	Registers.AddUninitialized(RegisterCount);
+
+	if (StartInstance == 0)
+	{
+		FMemory::Memcpy(&Registers[NumRegisters], RegisterTable.GetData(), sizeof(uint8*) * RegisterCount);
+	}
+	else
+	{
+		int32 StartOffset = sizeof(float) * StartInstance;
+		for (int32 i = 0; i < RegisterCount; ++i)
+		{
+			Registers[NumRegisters + i] = RegisterTable[i] + StartOffset;
+		}
+	}
+}
+
+FORCEINLINE void FNiagaraDataBuffer::ClearRegisterTable(FNiagaraRegisterTable& Registers)
+{
+	int32 RegisterCount = RegisterTable.Num();
+	Registers.AddZeroed(RegisterCount);
+}
diff --git a/Engine/Plugins/FX/Niagara/Source/Niagara/Classes/NiagaraScript.h b/Engine/Plugins/FX/Niagara/Source/Niagara/Classes/NiagaraScript.h
--- a/Engine/Plugins/FX/Niagara/Source/Niagara/Classes/NiagaraScript.h
+++ b/Engine/Plugins/FX/Niagara/Source/Niagara/Classes/NiagaraScript.h
@@ -215,6 +215,10 @@ public:
 	UPROPERTY()
 	TArray<uint8> ByteCode;
 
+	/** Number of temp registers used by this script. */
+	UPROPERTY()
+	int32 NumTempRegisters;
+
 	/** Number of user pointers we must pass to the VM. */
 	UPROPERTY()
 	int32 NumUserPtrs;
diff --git a/Engine/Plugins/FX/Niagara/Source/Niagara/Private/NiagaraDataSet.cpp b/Engine/Plugins/FX/Niagara/Source/Niagara/Private/NiagaraDataSet.cpp
--- a/Engine/Plugins/FX/Niagara/Source/Niagara/Private/NiagaraDataSet.cpp
+++ b/Engine/Plugins/FX/Niagara/Source/Niagara/Private/NiagaraDataSet.cpp
@@ -323,20 +323,6 @@ uint32 FNiagaraDataSet::GetSizeBytes()const
 	return Size;
 }
 
-void FNiagaraDataSet::ClearRegisterTable(uint8** Registers, int32& NumRegisters)
-{
-	for (FNiagaraVariableLayoutInfo& VarLayout : VariableLayouts)
-	{
-		int32 NumComps = VarLayout.GetNumFloatComponents() + VarLayout.GetNumInt32Components();
-		for (int32 CompIdx = 0; CompIdx < NumComps; ++CompIdx)
-		{
-			Registers[NumRegisters + CompIdx] = nullptr;
-		}
-		NumRegisters += NumComps;
-	}
-}
-
-
 void FNiagaraDataSet::CheckForNaNs()const
 {
 	for (const FNiagaraDataBuffer* Buffer : Data)
@@ -411,6 +397,8 @@ void FNiagaraDataSet::Finalize()
 	BuildLayout();
 
 	ResetBuffers();
+
+	GetCurrentDataChecked().BuildRegisterTable();
 }
 
 const FNiagaraVariableLayoutInfo* FNiagaraDataSet::GetVariableLayout(const FNiagaraVariable& Var)const
@@ -524,22 +512,6 @@ FNiagaraDataBuffer::~FNiagaraDataBuffer()
 	DEC_MEMORY_STAT_BY(STAT_NiagaraParticleMemory, FloatData.GetAllocatedSize() + Int32Data.GetAllocatedSize());
 }
 
-void FNiagaraDataBuffer::CheckUsage(bool bReadOnly)const
-{
-	check(Owner);
-	if (Owner->SimTarget == ENiagaraSimTarget::CPUSim)
-	{
-		//We can read on the RT but any modifications must be GT (or GT Task).
-		check(IsInGameThread() || (bReadOnly || !IsInRenderingThread()));
-	}
-	else
-	{
-		check(Owner->SimTarget == ENiagaraSimTarget::GPUComputeSim);
-		//Everything other than init for GPU sims should be done on the RT.
-		check(IsInRenderingThread());
-	}
-}
-
 int32 FNiagaraDataBuffer::TransferInstance(FNiagaraDataBuffer& SourceBuffer, int32 InstanceIndex, bool bRemoveFromSource)
 {
 	CheckUsage(false);
@@ -662,6 +634,9 @@ void FNiagaraDataBuffer::Allocate(uint32 InNumInstances, bool bMaintainExisting)
 	{
 		IDToIndexTable.Reset();
 	}
+
+	//Must rebuild our register table every time we allocate new buffers.
+	BuildRegisterTable();
 }
 
 void FNiagaraDataBuffer::AllocateGPU(uint32 InNumInstances, FNiagaraGPUInstanceCountManager& GPUInstanceCountManager, FRHICommandList &RHICmdList)
@@ -958,33 +933,6 @@ void FNiagaraDataBuffer::Dump(int32 StartIndex, int32 InNumInstances, const FStr
 
 }
 
-bool FNiagaraDataBuffer::AppendToRegisterTable(uint8** Registers, int32& NumRegisters, int32 StartInstance)
-{
-	check(Owner && Owner->IsInitialized());
-	check(Owner->GetSimTarget() == ENiagaraSimTarget::CPUSim);
-	CheckUsage(true);
-
-	for (FNiagaraVariableLayoutInfo& VarLayout : Owner->VariableLayouts)
-	{
-		int32 NumFloats = VarLayout.GetNumFloatComponents();
-		int32 NumInts = VarLayout.GetNumInt32Components();
-		for (int32 CompIdx = 0; CompIdx < NumFloats; ++CompIdx)
-		{
-			uint32 CompBufferOffset = VarLayout.FloatComponentStart + CompIdx;
-			uint32 CompRegisterOffset = VarLayout.LayoutInfo.FloatComponentRegisterOffsets[CompIdx];
-			Registers[NumRegisters + CompRegisterOffset] = (uint8*)GetInstancePtrFloat(CompBufferOffset, StartInstance);
-		}
-		for (int32 CompIdx = 0; CompIdx < NumInts; ++CompIdx)
-		{
-			uint32 CompBufferOffset = VarLayout.Int32ComponentStart + CompIdx;
-			uint32 CompRegisterOffset = VarLayout.LayoutInfo.Int32ComponentRegisterOffsets[CompIdx];
-			Registers[NumRegisters + CompRegisterOffset] = (uint8*)GetInstancePtrInt32(CompBufferOffset, StartInstance);
-		}
-		NumRegisters += NumFloats + NumInts;
-	}
-	return true;
-}
-
 /////////////////////////////////////////////////////////////////////////
 
 void FNiagaraDataBuffer::SetShaderParams(FNiagaraShader *Shader, FRHICommandList &CommandList, bool bInput)
@@ -1135,3 +1083,31 @@ void FScopedNiagaraDataSetGPUReadback::ReadbackData(NiagaraEmitterInstanceBatche
 	);
 	FlushRenderingCommands();
 }
+
+//////////////////////////////////////////////////////////////////////////
+
+void FNiagaraDataBuffer::BuildRegisterTable()
+{
+	int32 TotalRegisters = Owner->GetNumFloatComponents() + Owner->GetNumInt32Components();
+	RegisterTable.Reset();
+	RegisterTable.SetNumUninitialized(TotalRegisters);
+	int32 NumRegisters = 0;
+	for (FNiagaraVariableLayoutInfo& VarLayout : Owner->VariableLayouts)
+	{
+		int32 NumFloats = VarLayout.GetNumFloatComponents();
+		int32 NumInts = VarLayout.GetNumInt32Components();
+		for (int32 CompIdx = 0; CompIdx < NumFloats; ++CompIdx)
+		{
+			uint32 CompBufferOffset = VarLayout.FloatComponentStart + CompIdx;
+			uint32 CompRegisterOffset = VarLayout.LayoutInfo.FloatComponentRegisterOffsets[CompIdx];
+			RegisterTable[NumRegisters + CompRegisterOffset] = (uint8*)GetComponentPtrFloat(CompBufferOffset);
+		}
+		for (int32 CompIdx = 0; CompIdx < NumInts; ++CompIdx)
+		{
+			uint32 CompBufferOffset = VarLayout.Int32ComponentStart + CompIdx;
+			uint32 CompRegisterOffset = VarLayout.LayoutInfo.Int32ComponentRegisterOffsets[CompIdx];
+			RegisterTable[NumRegisters + CompRegisterOffset] = (uint8*)GetComponentPtrInt32(CompBufferOffset);
+		}
+		NumRegisters += NumFloats + NumInts;
+	}
+}
diff --git a/Engine/Plugins/FX/Niagara/Source/Niagara/Private/NiagaraScriptExecutionContext.cpp b/Engine/Plugins/FX/Niagara/Source/Niagara/Private/NiagaraScriptExecutionContext.cpp
--- a/Engine/Plugins/FX/Niagara/Source/Niagara/Private/NiagaraScriptExecutionContext.cpp
+++ b/Engine/Plugins/FX/Niagara/Source/Niagara/Private/NiagaraScriptExecutionContext.cpp
@@ -171,10 +171,8 @@ bool FNiagaraScriptExecutionContext::Execute(uint32 NumInstances)
 
 	++TickCounter;//Should this be per execution?
 
-	int32 NumInputRegisters = 0;
-	int32 NumOutputRegisters = 0;
-	uint8* InputRegisters[VectorVM::MaxInputRegisters];
-	uint8* OutputRegisters[VectorVM::MaxOutputRegisters];
+	FNiagaraRegisterTable InputRegisters;
+	FNiagaraRegisterTable OutputRegisters;
 
 	DataSetMetaTable.Reset();
 
@@ -188,45 +186,41 @@ bool FNiagaraScriptExecutionContext::Execute(uint32 NumInstances)
 			DataSetInfo.DataSet->CheckForNaNs();
 #endif
 			check(Info.DataSet);
-			DataSetMetaTable.Emplace(&InputRegisters[NumInputRegisters], NumInputRegisters, Info.StartInstance,
+			int32 NumInputRegisters = InputRegisters.Num();
+			DataSetMetaTable.Emplace(NumInputRegisters, Info.StartInstance,
 				DestinationData ? &DestinationData->GetIDTable() : nullptr, &Info.DataSet->GetFreeIDTable(), &Info.DataSet->GetNumFreeIDs(), &Info.DataSet->GetMaxUsedID(), Info.DataSet->GetIDAcquireTag());
 
 			int32 TotalComponents = Info.DataSet->GetNumFloatComponents() + Info.DataSet->GetNumInt32Components();
-			if (NumInputRegisters + TotalComponents > VectorVM::MaxInputRegisters || NumOutputRegisters + TotalComponents > VectorVM::MaxOutputRegisters)
-			{
-				UE_LOG(LogNiagara, Warning, TEXT("VM Script is using too many registers."));//TODO: We can make this a non issue and support any number of registers.
-				bError = true;
-				break;
-			}
-
 			if (Info.Input)
 			{
-				Info.Input->AppendToRegisterTable(InputRegisters, NumInputRegisters, Info.StartInstance);
+				Info.Input->AppendToRegisterTable(InputRegisters, Info.StartInstance);
 			}
 			else
 			{
-				Info.DataSet->ClearRegisterTable(InputRegisters, NumInputRegisters);
+				Info.DataSet->GetCurrentDataChecked().ClearRegisterTable(InputRegisters);
 			}
 
 			if (Info.Output)
 			{
-				Info.Output->AppendToRegisterTable(OutputRegisters, NumOutputRegisters, Info.StartInstance);
+				Info.Output->AppendToRegisterTable(OutputRegisters, Info.StartInstance);
 			}
 			else
 			{
-				Info.DataSet->ClearRegisterTable(OutputRegisters, NumOutputRegisters);
+				Info.DataSet->GetCurrentDataChecked().ClearRegisterTable(OutputRegisters);
 			}
 		}
 	}
 
 	if (GbExecVMScripts != 0 && !bError)
 	{
+		const FNiagaraVMExecutableData& ExecData = Script->GetVMExecutableData();
 		VectorVM::Exec(
-			Script->GetVMExecutableData().ByteCode.GetData(),
-			InputRegisters,
-			NumInputRegisters,
-			OutputRegisters,
-			NumOutputRegisters,
+			ExecData.ByteCode.GetData(),
+			ExecData.NumTempRegisters,
+			InputRegisters.GetData(),
+			InputRegisters.Num(),
+			OutputRegisters.GetData(),
+			OutputRegisters.Num(),
 			Parameters.GetParameterDataArray().GetData(),
 			DataSetMetaTable,
 			FunctionTable.GetData(),
diff --git a/Engine/Plugins/FX/Niagara/Source/NiagaraCore/Private/NiagaraCustomVersion.cpp b/Engine/Plugins/FX/Niagara/Source/NiagaraCore/Private/NiagaraCustomVersion.cpp
--- a/Engine/Plugins/FX/Niagara/Source/NiagaraCore/Private/NiagaraCustomVersion.cpp
+++ b/Engine/Plugins/FX/Niagara/Source/NiagaraCore/Private/NiagaraCustomVersion.cpp
@@ -8,6 +8,4 @@ const FGuid FNiagaraCustomVersion::GUID(0xFCF57AFA, 0x50764283, 0xB9A9E658, 0xFF
 // Register the custom version with core
 FCustomVersionRegistration GRegisterNiagaraCustomVersion(FNiagaraCustomVersion::GUID, FNiagaraCustomVersion::LatestVersion, TEXT("NiagaraVer"));
 
-const FGuid FNiagaraCustomVersion::LatestScriptCompileVersion(0xE2AC5117, 0x5B44424D, 0x91265F00, 0xA1AEB559);
-
-
+const FGuid FNiagaraCustomVersion::LatestScriptCompileVersion(0x4F951D5B, 0x13E94DCB, 0x94C31CEA, 0x64D923C9);
\ No newline at end of file
diff --git a/Engine/Plugins/FX/Niagara/Source/NiagaraEditor/Private/NiagaraCompiler.cpp b/Engine/Plugins/FX/Niagara/Source/NiagaraEditor/Private/NiagaraCompiler.cpp
--- a/Engine/Plugins/FX/Niagara/Source/NiagaraEditor/Private/NiagaraCompiler.cpp
+++ b/Engine/Plugins/FX/Niagara/Source/NiagaraEditor/Private/NiagaraCompiler.cpp
@@ -834,6 +834,7 @@ FNiagaraCompileResults FHlslNiagaraCompiler::CompileScript(const FNiagaraCompile
 			check(TranslatorOutput);
 			*CompileResults.Data.Get() = TranslatorOutput->ScriptData;
 			CompileResults.Data->ByteCode = CompilationOutput.ByteCode;
+			CompileResults.Data->NumTempRegisters = CompilationOutput.MaxTempRegistersUsed + 1;
 			CompileResults.Data->LastAssemblyTranslation = CompilationOutput.AssemblyAsString;
 			CompileResults.Data->LastOpCount = CompilationOutput.NumOps;
 			//Build internal parameters
diff --git a/Engine/Source/Developer/ShaderFormatVectorVM/Private/ir_vm_gen_bytecode_visitor.cpp b/Engine/Source/Developer/ShaderFormatVectorVM/Private/ir_vm_gen_bytecode_visitor.cpp
--- a/Engine/Source/Developer/ShaderFormatVectorVM/Private/ir_vm_gen_bytecode_visitor.cpp
+++ b/Engine/Source/Developer/ShaderFormatVectorVM/Private/ir_vm_gen_bytecode_visitor.cpp
@@ -592,7 +592,7 @@ struct op_base
 
 	}
 
-	static bool finalize_component_temporary_allocation(_mesa_glsl_parse_state* parse_state, variable_info_node* component, const unsigned num_temp_registers, unsigned *registers, unsigned op_idx)
+	static bool finalize_component_temporary_allocation(_mesa_glsl_parse_state* parse_state, FVectorVMCompilationOutput& CompilationOutput, variable_info_node* component, const unsigned num_temp_registers, unsigned *registers, unsigned op_idx)
 	{
 		check(component);
 
@@ -615,6 +615,7 @@ struct op_base
 						{
 							component->offset = VectorVM::FirstTempRegister + j;
 							registers[j] = component->last_read;
+							CompilationOutput.MaxTempRegistersUsed = FMath::Max((int32)j, CompilationOutput.MaxTempRegistersUsed);
 #if VM_VERBOSE_LOGGING
 							UE_LOG(LogVVMBackend, Log, TEXT("OP:%d | Comp:%p allocated Reg: %d | Last Read: %d |"), op_idx, component, j, component->last_read);
 #endif
@@ -769,7 +770,7 @@ struct op_standard : public op_base
 
 	virtual bool finalize(_mesa_glsl_parse_state* parse_state, FVectorVMCompilationOutput&CompilationOutput, const unsigned num_temp_registers, unsigned *registers, unsigned op_idx)
 	{
-		return finalize_component_temporary_allocation(parse_state, dest_component, num_temp_registers, registers, op_idx);
+		return finalize_component_temporary_allocation(parse_state, CompilationOutput, dest_component, num_temp_registers, registers, op_idx);
 	}
 
 	virtual void validate(_mesa_glsl_parse_state* parse_state, unsigned op_idx)override
@@ -827,7 +828,7 @@ struct op_input : public op_base
 
 	virtual bool finalize(_mesa_glsl_parse_state* parse_state, FVectorVMCompilationOutput&CompilationOutput, const unsigned num_temp_registers, unsigned *registers, unsigned op_idx)
 	{
-		return finalize_component_temporary_allocation(parse_state, dest_component, num_temp_registers, registers, op_idx);
+		return finalize_component_temporary_allocation(parse_state, CompilationOutput, dest_component, num_temp_registers, registers, op_idx);
 	}
 
 	virtual void validate(_mesa_glsl_parse_state* parse_state, unsigned op_idx)override
@@ -874,7 +875,7 @@ struct op_input_noadvance : public op_input
 
 	virtual bool finalize(_mesa_glsl_parse_state* parse_state, FVectorVMCompilationOutput&CompilationOutput, const unsigned num_temp_registers, unsigned *registers, unsigned op_idx)
 	{
-		return finalize_component_temporary_allocation(parse_state, dest_component, num_temp_registers, registers, op_idx);
+		return finalize_component_temporary_allocation(parse_state, CompilationOutput, dest_component, num_temp_registers, registers, op_idx);
 	}
 
 	virtual void validate(_mesa_glsl_parse_state* parse_state, unsigned op_idx)override
@@ -967,7 +968,7 @@ struct op_index_acquire : public op_base
 
 	virtual bool finalize(_mesa_glsl_parse_state* parse_state, FVectorVMCompilationOutput&CompilationOutput, const unsigned num_temp_registers, unsigned *registers, unsigned op_idx)
 	{
-		return finalize_component_temporary_allocation(parse_state, dest_component, num_temp_registers, registers, op_idx);
+		return finalize_component_temporary_allocation(parse_state, CompilationOutput, dest_component, num_temp_registers, registers, op_idx);
 	}
 
 	virtual void validate(_mesa_glsl_parse_state* parse_state, unsigned op_idx)override
@@ -1010,12 +1011,12 @@ struct op_id_acquire : public op_base
 
 	virtual bool finalize(_mesa_glsl_parse_state* parse_state, FVectorVMCompilationOutput&CompilationOutput, const unsigned num_temp_registers, unsigned *registers, unsigned op_idx)
 	{
-		if (!finalize_component_temporary_allocation(parse_state, id_index_component, num_temp_registers, registers, op_idx))
+		if (!finalize_component_temporary_allocation(parse_state, CompilationOutput, id_index_component, num_temp_registers, registers, op_idx))
 		{
 			_mesa_glsl_error(parse_state, "Failed to allocate temporary registers");
 			return false;
 		}
-		if (!finalize_component_temporary_allocation(parse_state, id_tag_component, num_temp_registers, registers, op_idx))
+		if (!finalize_component_temporary_allocation(parse_state, CompilationOutput, id_tag_component, num_temp_registers, registers, op_idx))
 		{
 			_mesa_glsl_error(parse_state, "Failed to allocate temporary registers");
 			return false;
@@ -1175,6 +1176,7 @@ struct op_external_func : public op_base
 	struct finalize_temp_register_ctx
 	{
 		_mesa_glsl_parse_state* parse_state;
+		FVectorVMCompilationOutput* CompilationOutput;
 		unsigned num_temp_registers;
 		unsigned *registers;
 		unsigned op_idx;
@@ -1187,7 +1189,7 @@ struct op_external_func : public op_base
 		check(node->is_scalar());
 
 		//Some outputs for an external function can have invalid locations. If the thing has *any* valid outputs then it's considered a success.
-		ctx->success |= finalize_component_temporary_allocation(ctx->parse_state, node, ctx->num_temp_registers, ctx->registers, ctx->op_idx);
+		ctx->success |= finalize_component_temporary_allocation(ctx->parse_state, *ctx->CompilationOutput, node, ctx->num_temp_registers, ctx->registers, ctx->op_idx);
 	}
 
 	virtual bool finalize(_mesa_glsl_parse_state* parse_state, FVectorVMCompilationOutput&CompilationOutput, const unsigned num_temp_registers, unsigned *registers, unsigned op_idx)
@@ -1212,6 +1214,7 @@ struct op_external_func : public op_base
 
 		finalize_temp_register_ctx ctx;
 		ctx.parse_state = parse_state;
+		ctx.CompilationOutput = &CompilationOutput;
 		ctx.num_temp_registers = num_temp_registers;
 		ctx.registers = registers;
 		ctx.op_idx = op_idx;
diff --git a/Engine/Source/Developer/ShaderFormatVectorVM/Public/ShaderFormatVectorVM.h b/Engine/Source/Developer/ShaderFormatVectorVM/Public/ShaderFormatVectorVM.h
--- a/Engine/Source/Developer/ShaderFormatVectorVM/Public/ShaderFormatVectorVM.h
+++ b/Engine/Source/Developer/ShaderFormatVectorVM/Public/ShaderFormatVectorVM.h
@@ -7,10 +7,12 @@
 /** Data which is generated from the hlsl by the VectorVMBackend and fed back into the */
 struct FVectorVMCompilationOutput
 {
-	FVectorVMCompilationOutput(): NumOps(0){}
+	FVectorVMCompilationOutput(): MaxTempRegistersUsed(0), NumOps(0) {}
 
 	TArray<uint8> ByteCode;
 
+	int32 MaxTempRegistersUsed;
+
 	TArray<int32> InternalConstantOffsets;
 	TArray<uint8> InternalConstantData;
 	TArray<EVectorVMBaseTypes> InternalConstantTypes;
diff --git a/Engine/Source/Runtime/VectorVM/Private/VectorVM.cpp b/Engine/Source/Runtime/VectorVM/Private/VectorVM.cpp
--- a/Engine/Source/Runtime/VectorVM/Private/VectorVM.cpp
+++ b/Engine/Source/Runtime/VectorVM/Private/VectorVM.cpp
@@ -197,6 +197,8 @@ public:
 
 //////////////////////////////////////////////////////////////////////////
 
+TLockFreePointerListLIFO<FVectorVMContext> FVectorVMContext::VMContextPool;
+
 FVectorVMContext::FVectorVMContext()
 	: Code(nullptr)
 	, ConstantTable(nullptr)
@@ -211,34 +213,31 @@ FVectorVMContext::FVectorVMContext()
 	, StatScopes(nullptr)
 #endif
 {
-	uint32 TempRegisterSize = Align((GParallelVVMInstancesPerChunk) * VectorVM::MaxInstanceSizeBytes, VECTOR_WIDTH_BYTES) + VECTOR_WIDTH_BYTES;
-	TempRegTable.SetNumUninitialized(TempRegisterSize * VectorVM::NumTempRegisters);
-	// Map temporary registers.
-	for (int32 i = 0; i < VectorVM::NumTempRegisters; ++i)
-	{
-		RegisterTable[i] = TempRegTable.GetData() + TempRegisterSize * i;
-	}
-
 	RandStream.GenerateNewSeed();
 }
 
 void FVectorVMContext::PrepareForExec(
 	uint8*RESTRICT*RESTRICT InputRegisters,
 	uint8*RESTRICT*RESTRICT OutputRegisters,
-	int32 NumInputRegisters,
-	int32 NumOutputRegisters,
+	int32 InNumTempRegisters,
+	int32 InNumInputRegisters,
+	int32 InNumOutputRegisters,
 	const uint8* InConstantTable,
 	int32 *InDataSetIndexTable,
 	int32 *InDataSetOffsetTable,
 	int32 InNumSecondaryDatasets,
 	FVMExternalFunction* InExternalFunctionTable,
 	void** InUserPtrTable,
-	TArray<FDataSetMeta>& RESTRICT InDataSetMetaTable
+	TArray<FDataSetMeta>& RESTRICT InDataSetMetaTable,
+	int32 MaxNumInstances
 #if STATS
 	, const TArray<TStatId>* InStatScopes
 #endif
 )
 {
+	NumTempRegisters = InNumTempRegisters;
+	NumInputRegisters = InNumInputRegisters;
+	NumOutputRegisters = InNumOutputRegisters;
 	ConstantTable = InConstantTable;
 	DataSetIndexTable = InDataSetIndexTable;
 	DataSetOffsetTable = InDataSetOffsetTable;
@@ -251,6 +250,14 @@ void FVectorVMContext::PrepareForExec(
 	StatCounterStack.Reserve(StatScopes->Num());
 #endif
 
+	uint32 TempRegisterSize = Align(MaxNumInstances * VectorVM::MaxInstanceSizeBytes, VECTOR_WIDTH_BYTES) + VECTOR_WIDTH_BYTES;
+	TempRegTable.SetNumUninitialized(TempRegisterSize * NumTempRegisters, false);
+	// Attempt to map temp registers more tightly packed for low instance counts to reduce cache misses.
+	for (int32 i = 0; i < NumTempRegisters; ++i)
+	{
+		RegisterTable[i] = TempRegTable.GetData() + TempRegisterSize * i;
+	}
+
 	//Map IO Registers
 	for (int32 i = 0; i < NumInputRegisters; ++i)
 	{
@@ -1813,6 +1820,7 @@ void VectorVM::Init()
 
 void VectorVM::Exec(
 	uint8 const* Code,
+	int32 NumTempRegisters,
 	uint8** InputRegisters,
 	int32 NumInputRegisters,
 	uint8** OutputRegisters,
@@ -1842,6 +1850,7 @@ void VectorVM::Exec(
 		DataSetIndexTable.Add(DataSetMetaTable[Idx].DataSetAccessIndex);	// prime counter index table with the data set offset; will be incremented with every write for each instance
 	}
 
+	int32 MaxInstances = FMath::Min(GParallelVVMInstancesPerChunk, NumInstances);
 	int32 NumChunks = (NumInstances / GParallelVVMInstancesPerChunk) + 1;
 	int32 ChunksPerBatch = (GbParallelVVM != 0 && FApp::ShouldUseThreadingForPerformance()) ? GParallelVVMChunksPerBatch : NumChunks;
 	int32 NumBatches = FMath::DivideAndRoundUp(NumChunks, ChunksPerBatch);
@@ -1851,9 +1860,11 @@ void VectorVM::Exec(
 	{
 		//SCOPE_CYCLE_COUNTER(STAT_VVMExecChunk);
 
-		FVectorVMContext& Context = FVectorVMContext::Get();
-		Context.PrepareForExec(InputRegisters, OutputRegisters, NumInputRegisters, NumOutputRegisters, ConstantTable, DataSetIndexTable.GetData(), DataSetOffsetTable.GetData(), DataSetOffsetTable.Num(),
-			ExternalFunctionTable, UserPtrTable, DataSetMetaTable
+		FVectorVMContext* ContextPtr = FVectorVMContext::GetContext();
+		checkSlow(ContextPtr);
+		FVectorVMContext& Context = *ContextPtr;
+		Context.PrepareForExec(InputRegisters, OutputRegisters, NumTempRegisters, NumInputRegisters, NumOutputRegisters, ConstantTable, DataSetIndexTable.GetData(), DataSetOffsetTable.GetData(), DataSetOffsetTable.Num(),
+			ExternalFunctionTable, UserPtrTable, DataSetMetaTable, MaxInstances
 #if STATS
 			, &StatScopes
 #endif
@@ -2004,6 +2015,7 @@ void VectorVM::Exec(
 		}
 
 		Context.FinishExec();
+		FVectorVMContext::ReleaseContext(ContextPtr);
 	};
 
 	ParallelFor(NumBatches, ExecChunkBatch, GbParallelVVM == 0 || !bParallel);
diff --git a/Engine/Source/Runtime/VectorVM/Public/VectorVM.h b/Engine/Source/Runtime/VectorVM/Public/VectorVM.h
--- a/Engine/Source/Runtime/VectorVM/Public/VectorVM.h
+++ b/Engine/Source/Runtime/VectorVM/Public/VectorVM.h
@@ -154,7 +154,6 @@ enum class EVectorVMOp : uint8
 //Some require RWBuffer like support.
 struct FDataSetMeta
 {
-	uint8 **InputRegisters;
 	int32 RegisterOffset;
 	int32 DataSetAccessIndex;	// index for individual elements of this set
 
@@ -180,14 +179,14 @@ struct FDataSetMeta
 	FORCEINLINE void LockFreeTable();
 	FORCEINLINE void UnlockFreeTable();
 
-	FDataSetMeta(uint8 **Data, int32 InRegisterOffset, int32 InInstanceOffset, TArray<int32>* InIDTable, TArray<int32>* InFreeIDTable, int32* InNumFreeIDs, int32* InMaxUsedID, int32 InIDAcquireTag)
-		: InputRegisters(Data), RegisterOffset(InRegisterOffset), DataSetAccessIndex(INDEX_NONE), InstanceOffset(InInstanceOffset)
+	FDataSetMeta(int32 InRegisterOffset, int32 InInstanceOffset, TArray<int32>* InIDTable, TArray<int32>* InFreeIDTable, int32* InNumFreeIDs, int32* InMaxUsedID, int32 InIDAcquireTag)
+		: RegisterOffset(InRegisterOffset), DataSetAccessIndex(INDEX_NONE), InstanceOffset(InInstanceOffset)
 		, IDTable(InIDTable), FreeIDTable(InFreeIDTable), NumFreeIDs(InNumFreeIDs), MaxUsedID(InMaxUsedID), IDAcquireTag(InIDAcquireTag)
 	{
 	}
 
 	FDataSetMeta() 
-		: InputRegisters(nullptr), RegisterOffset(0), DataSetAccessIndex(INDEX_NONE), InstanceOffset(0)
+		: RegisterOffset(0), DataSetAccessIndex(INDEX_NONE), InstanceOffset(0)
 		, IDTable(nullptr), FreeIDTable(nullptr), NumFreeIDs(nullptr), MaxUsedID(nullptr), IDAcquireTag(0)
 	{}
 
@@ -234,12 +233,18 @@ struct FDataSetThreadLocalTempData
 /**
 * Context information passed around during VM execution.
 */
-struct FVectorVMContext : TThreadSingleton<FVectorVMContext>
+struct FVectorVMContext
 {
 	/** Pointer to the next element in the byte code. */
 	uint8 const* RESTRICT Code;
 	/** Pointer to the constant table. */
 	uint8 const* RESTRICT ConstantTable;
+	/** Num temp registers required by this script. */
+	int32 NumTempRegisters;
+	/** Num input required by this script. */
+	int32 NumInputRegisters;
+	/** Num output required by this script. */
+	int32 NumOutputRegisters;
 	/** Pointer to the data set index counter table */
 	int32* RESTRICT DataSetIndexTable;
 	int32* RESTRICT DataSetOffsetTable;
@@ -264,7 +269,7 @@ struct FVectorVMContext : TThreadSingleton<FVectorVMContext>
 #endif
 
 	TArray<uint8, TAlignedHeapAllocator<VECTOR_WIDTH_BYTES>> TempRegTable;
-	uint8 *RESTRICT RegisterTable[VectorVM::MaxRegisters];
+	uint8*RESTRICT RegisterTable[VectorVM::MaxRegisters];
 
 	/** Thread local random stream for use in external functions needing non-deterministic randoms. */
 	FRandomStream RandStream;
@@ -277,6 +282,7 @@ struct FVectorVMContext : TThreadSingleton<FVectorVMContext>
 	void PrepareForExec(
 		uint8*RESTRICT*RESTRICT InputRegisters,
 		uint8*RESTRICT*RESTRICT OutputRegisters,
+		int32 InNumTempRegisters,
 		int32 NumInputRegisters,
 		int32 NumOutputRegisters,
 		const uint8* InConstantTable,
@@ -285,7 +291,8 @@ struct FVectorVMContext : TThreadSingleton<FVectorVMContext>
 		int32 InNumSecondaryDatasets,
 		FVMExternalFunction* InExternalFunctionTable,
 		void** InUserPtrTable,
-		TArray<FDataSetMeta>& RESTRICT InDataSetMetaTable
+		TArray<FDataSetMeta>& RESTRICT InDataSetMetaTable,
+		int32 MaxNumInstances
 #if STATS
 		, const TArray<TStatId>* InStatScopes
 #endif
@@ -302,8 +309,30 @@ struct FVectorVMContext : TThreadSingleton<FVectorVMContext>
 		RandCounters.Reset();
 		RandCounters.SetNumZeroed(InNumInstances);
 	}
+
+	static TLockFreePointerListLIFO<FVectorVMContext> VMContextPool;
+	FORCEINLINE static FVectorVMContext* GetContext();
+	FORCEINLINE static void ReleaseContext(FVectorVMContext* Context);
 };
 
+FORCEINLINE FVectorVMContext* FVectorVMContext::GetContext()
+{
+	if (FVectorVMContext* Ret = VMContextPool.Pop())
+	{
+		return Ret;
+	}
+	else
+	{
+		return new FVectorVMContext();
+	}
+}
+
+FORCEINLINE void FVectorVMContext::ReleaseContext(FVectorVMContext* Context)
+{
+	checkSlow(Context);
+	VMContextPool.Push(Context);
+}
+
 namespace VectorVM
 {
 	/** Get total number of op-codes */
@@ -321,6 +350,7 @@ namespace VectorVM
 	 */
 	VECTORVM_API void Exec(
 		uint8 const* Code,
+		int32 NumTempRegisters,
 		uint8** InputRegisters,
 		int32 NumInputRegisters,
 		uint8** OutputRegisters,
