enable FMA and Subtraction in favor of negated addition on Niagara Shader backend.
Also run another optimization path after scalarization
Originally authored by [REDACTED_PERSON], integrated by [REDACTED_PERSON]

#tests Engine tests passed

diff --git a/Engine/Plugins/FX/Niagara/Source/NiagaraCore/Private/NiagaraCustomVersion.cpp b/Engine/Plugins/FX/Niagara/Source/NiagaraCore/Private/NiagaraCustomVersion.cpp
--- a/Engine/Plugins/FX/Niagara/Source/NiagaraCore/Private/NiagaraCustomVersion.cpp
+++ b/Engine/Plugins/FX/Niagara/Source/NiagaraCore/Private/NiagaraCustomVersion.cpp
@@ -8,4 +8,4 @@ const FGuid FNiagaraCustomVersion::GUID(0xFCF57AFA, 0x50764283, 0xB9A9E658, 0xFF
 // Register the custom version with core
 FCustomVersionRegistration GRegisterNiagaraCustomVersion(FNiagaraCustomVersion::GUID, FNiagaraCustomVersion::LatestVersion, TEXT("NiagaraVer"));
 
-const FGuid FNiagaraCustomVersion::LatestScriptCompileVersion(0x4F951D5B, 0x13E94DCB, 0x94C31CEA, 0x64D923C9);
\ No newline at end of file
+const FGuid FNiagaraCustomVersion::LatestScriptCompileVersion(0x31F56712, 0xC9264924, 0x9B11B480, 0x1499071A);
diff --git a/Engine/Source/Developer/ShaderFormatVectorVM/Private/VectorVMBackend.cpp b/Engine/Source/Developer/ShaderFormatVectorVM/Private/VectorVMBackend.cpp
--- a/Engine/Source/Developer/ShaderFormatVectorVM/Private/VectorVMBackend.cpp
+++ b/Engine/Source/Developer/ShaderFormatVectorVM/Private/VectorVMBackend.cpp
@@ -227,7 +227,6 @@ char* FVectorVMCodeBackend::GenerateCode(exec_list* ir, _mesa_glsl_parse_state*
 	}
 
 	{
-
 		SCOPE_CYCLE_COUNTER(STAT_VVMScalarize);
 		vm_debug_print("== Scalarize ==\n");
 		vm_scalarize_ops(ir, state);
@@ -237,6 +236,17 @@ char* FVectorVMCodeBackend::GenerateCode(exec_list* ir, _mesa_glsl_parse_state*
 		if (state->error) return nullptr;
 	}
 
+	{
+		SCOPE_CYCLE_COUNTER(STAT_VVMScalarize);
+		vm_debug_print("== Scalar Optimization ==\n");
+		progress = false;
+		do
+		{
+			progress = do_algebraic(state, ir);
+		} while (progress);
+		vm_debug_dump(ir, state);
+		if (state->error) return nullptr;
+	}
 	{
 		SCOPE_CYCLE_COUNTER(STAT_VVMMergeOps);
 		vm_debug_print("== Merge Ops ==\n");
diff --git a/Engine/Source/Developer/ShaderFormatVectorVM/Private/VectorVMBackend.h b/Engine/Source/Developer/ShaderFormatVectorVM/Private/VectorVMBackend.h
--- a/Engine/Source/Developer/ShaderFormatVectorVM/Private/VectorVMBackend.h
+++ b/Engine/Source/Developer/ShaderFormatVectorVM/Private/VectorVMBackend.h
@@ -14,8 +14,9 @@ class FVectorVMLanguageSpec : public ILanguageSpec
 protected:
 
 public:
-
-	virtual bool SupportsDeterminantIntrinsic() const override{	return false;	}
+	virtual bool LowerSubToAddNeg() const override { return false; }
+	virtual bool SupportsFusedMultiplyAdd() const override { return true; }
+	virtual bool SupportsDeterminantIntrinsic() const override { return false; }
 	virtual bool SupportsTransposeIntrinsic() const override{ return false;	}
 	virtual bool SupportsIntegerModulo() const override { return false;	}
 	virtual bool SupportsMatrixConversions() const override { return false; }
diff --git a/Engine/Source/Developer/ShaderFormatVectorVM/Private/ir_vm_gen_bytecode_visitor.cpp b/Engine/Source/Developer/ShaderFormatVectorVM/Private/ir_vm_gen_bytecode_visitor.cpp
--- a/Engine/Source/Developer/ShaderFormatVectorVM/Private/ir_vm_gen_bytecode_visitor.cpp
+++ b/Engine/Source/Developer/ShaderFormatVectorVM/Private/ir_vm_gen_bytecode_visitor.cpp
@@ -314,6 +314,10 @@ void BuildExpressionMap()
 
 		Info = &VMExpressionMap.Add(ir_ternop_lerp);
 		Info->Add(FVMExpresssionInfo(EVectorVMOp::lerp, glsl_type::float_type, glsl_type::float_type, glsl_type::float_type, glsl_type::float_type));
+
+		Info = &VMExpressionMap.Add(ir_ternop_fma);
+		Info->Add(FVMExpresssionInfo(EVectorVMOp::mad, glsl_type::float_type, glsl_type::float_type, glsl_type::float_type, glsl_type::float_type));
+
 		//Info = &VMExpressionMap.Add(ir_ternop_smoothstep);TODO: Smoothstep
 		Info = &VMExpressionMap.Add(ir_ternop_clamp);
 		Info->Add(FVMExpresssionInfo(EVectorVMOp::clamp, glsl_type::float_type, glsl_type::float_type, glsl_type::float_type, glsl_type::float_type));
@@ -346,56 +350,57 @@ EVectorVMOp get_special_vm_opcode(ir_function_signature* signature)
 	}
 	else if (strcmp(signature->function_name(), "Modulo") == 0)
 	{
-		VVMOpCode = EVectorVMOp::fmod;
+		return EVectorVMOp::fmod;
 	}
 	else if (strcmp(signature->function_name(), "select") == 0)
 	{
-		VVMOpCode = EVectorVMOp::select;
+		return EVectorVMOp::select;
 	}
 	else if (strcmp(signature->function_name(), "noise") == 0)
 	{
-		VVMOpCode = EVectorVMOp::noise;
+		return EVectorVMOp::noise;
 	}
 	else if (strncmp(signature->function_name(), "InputDataNoadvance", strlen("InputDataNoadvance")) == 0)
 	{
-		VVMOpCode = EVectorVMOp::inputdata_noadvance_32bit;
+		return EVectorVMOp::inputdata_noadvance_32bit;
 	}
 	else if (strncmp(signature->function_name(), "InputData", strlen("InputData")) == 0)
 	{
-		VVMOpCode = EVectorVMOp::inputdata_32bit;
+		return EVectorVMOp::inputdata_32bit;
 	}
 	else if (strncmp(signature->function_name(), "OutputData", strlen("OutputData")) == 0)
 	{
-		VVMOpCode = EVectorVMOp::outputdata_32bit;
+		return EVectorVMOp::outputdata_32bit;
 	}
 	else if (strcmp(signature->function_name(), "AcquireIndex") == 0)
 	{
-		VVMOpCode = EVectorVMOp::acquireindex;
+		return EVectorVMOp::acquireindex;
 	}
 	else if (strcmp(signature->function_name(), "AcquireID") == 0)
 	{
-		VVMOpCode = EVectorVMOp::acquire_id;
+		return EVectorVMOp::acquire_id;
 	}
 	else if (strcmp(signature->function_name(), "UpdateID") == 0)
 	{
-		VVMOpCode = EVectorVMOp::update_id;
+		return EVectorVMOp::update_id;
 	}
 	else if (strcmp(signature->function_name(), "ExecIndex") == 0)
 	{
-		VVMOpCode = EVectorVMOp::exec_index;
+		return EVectorVMOp::exec_index;
 	}
 	else if (strcmp(signature->function_name(), "EnterStatScope") == 0)
 	{
-		VVMOpCode = EVectorVMOp::enter_stat_scope;
+		return EVectorVMOp::enter_stat_scope;
 	}
 	else if (strcmp(signature->function_name(), "ExitStatScope") == 0)
 	{
-		VVMOpCode = EVectorVMOp::exit_stat_scope;
+		return EVectorVMOp::exit_stat_scope;
 	}
 	else if (signature->body.get_head() == nullptr)
 	{
-		VVMOpCode = EVectorVMOp::external_func_call;
+		return EVectorVMOp::external_func_call;
 	}
+	check(VVMOpCode != EVectorVMOp::done);
 	return VVMOpCode;
 }
 
@@ -682,6 +687,7 @@ struct op_hook : public op_base
 
 	virtual void add_to_bytecode(TArray<uint8>& bytecode)
 	{
+		check(op_code != EVectorVMOp::done);
 		emit_byte((uint8)op_code, bytecode);
 		if (source_component)
 		{
@@ -748,6 +754,7 @@ struct op_standard : public op_base
 
 	virtual void add_to_bytecode(TArray<uint8>& bytecode)
 	{
+		check(op_code != EVectorVMOp::done);
 		emit_byte((uint8)op_code, bytecode);
 
 		if (num_operands > 0)
@@ -818,6 +825,7 @@ struct op_input : public op_base
 	{
 		check(dest_component);
 		//if(instance_idx_component)//TODO: Handle the case with an explicit instance index. Probably want a separate VMOp for this 
+		check(op_code != EVectorVMOp::done);
 		emit_byte((uint8)op_code, bytecode);
 		emit_short(dataset_idx, bytecode);
 		emit_short(register_idx + VectorVM::FirstInputRegister, bytecode);
@@ -865,6 +873,7 @@ struct op_input_noadvance : public op_input
 	{
 		check(dest_component);
 		//if(instance_idx_component)//TODO: Handle the case with an explicit instance index. Probably want a separate VMOp for this 
+		check(op_code != EVectorVMOp::done);
 		emit_byte((uint8)op_code, bytecode);
 		emit_short(dataset_idx, bytecode);
 		emit_short(register_idx + VectorVM::FirstInputRegister, bytecode);
@@ -915,6 +924,7 @@ struct op_output : public op_base
 	{
 		check(instance_idx_component);
 		check(value_component);
+		check(op_code != EVectorVMOp::done);
 		emit_byte((uint8)op_code, bytecode);
 		//value can be a constant or a register so have to emit a location byte.
 		emit_byte(VectorVM::CreateSrcOperandMask(value_component->owner->location), bytecode);
@@ -959,6 +969,7 @@ struct op_index_acquire : public op_base
 	{
 		check(valid_component);
 		check(dest_component);
+		check(op_code != EVectorVMOp::done);
 		emit_byte((uint8)op_code, bytecode);
 		emit_byte(VectorVM::CreateSrcOperandMask(valid_component->owner->location), bytecode);
 		emit_short(dataset_idx, bytecode);
@@ -1003,6 +1014,7 @@ struct op_id_acquire : public op_base
 	{
 		check(id_index_component);
 		check(id_tag_component);
+		check(op_code != EVectorVMOp::done);
 		emit_byte((uint8)op_code, bytecode);
 		emit_short(dataset_idx, bytecode);
 		emit_short(id_index_component->offset, bytecode);
@@ -1060,6 +1072,7 @@ struct op_id_update : public op_base
 	{
 		check(id_component);
 		check(index_component);
+		check(op_code != EVectorVMOp::done);
 		emit_byte((uint8)op_code, bytecode);
 		emit_short(dataset_idx, bytecode);
 		emit_short(id_component->offset, bytecode);
@@ -1660,7 +1673,7 @@ class ir_gen_vvm_visitor : public ir_hierarchical_visitor
 				op_standard* op = curr_op()->as_standard();
 				check(op);
 
-				if (op->op_code != EVectorVMOp::done)
+				if (op->op_code != EVectorVMOp::done && op->op_code != EVectorVMOp::mad)
 				{
 					//we're already inside another expression.
 					//Currently we only allow this for mad.
