vm visitor to stop vm ops reading from the value they're writing as this is potentially unsafe.

diff --git a/Engine/Plugins/FX/Niagara/Source/NiagaraCore/Private/NiagaraCustomVersion.cpp b/Engine/Plugins/FX/Niagara/Source/NiagaraCore/Private/NiagaraCustomVersion.cpp
--- a/Engine/Plugins/FX/Niagara/Source/NiagaraCore/Private/NiagaraCustomVersion.cpp
+++ b/Engine/Plugins/FX/Niagara/Source/NiagaraCore/Private/NiagaraCustomVersion.cpp
@@ -8,4 +8,4 @@ const FGuid FNiagaraCustomVersion::GUID(0xFCF57AFA, 0x50764283, 0xB9A9E658, 0xFF
 // Register the custom version with core
 FCustomVersionRegistration GRegisterNiagaraCustomVersion(FNiagaraCustomVersion::GUID, FNiagaraCustomVersion::LatestVersion, TEXT("NiagaraVer"));
 
-const FGuid FNiagaraCustomVersion::LatestScriptCompileVersion(0x31B222C7, 0x2AAF448A, 0xA9DCA5E7, 0x1DAE433D);
+const FGuid FNiagaraCustomVersion::LatestScriptCompileVersion(0x523B1F98, 0xF08242CD, 0x9815CB3C, 0xC7B5EC38);
diff --git a/Engine/Source/Developer/ShaderFormatVectorVM/Private/VectorVMBackend.cpp b/Engine/Source/Developer/ShaderFormatVectorVM/Private/VectorVMBackend.cpp
--- a/Engine/Source/Developer/ShaderFormatVectorVM/Private/VectorVMBackend.cpp
+++ b/Engine/Source/Developer/ShaderFormatVectorVM/Private/VectorVMBackend.cpp
@@ -267,6 +267,14 @@ char* FVectorVMCodeBackend::GenerateCode(exec_list* ir, _mesa_glsl_parse_state*
 		if (state->error) return nullptr;
 	}
 
+	{
+		//SCOPE_CYCLE_COUNTER(STAT_VVMCleanup);
+		vm_debug_print("== Remove self ref assignments ==\n");
+		vm_remove_self_ref_assignments(ir, state);
+		vm_debug_dump(ir, state);
+		if (state->error) return nullptr;
+	}
+
 	{
 		SCOPE_CYCLE_COUNTER(STAT_VVMCleanup);
 		vm_debug_print("== Cleanup ==\n");
diff --git a/Engine/Source/Developer/ShaderFormatVectorVM/Private/VectorVMBackend.h b/Engine/Source/Developer/ShaderFormatVectorVM/Private/VectorVMBackend.h
--- a/Engine/Source/Developer/ShaderFormatVectorVM/Private/VectorVMBackend.h
+++ b/Engine/Source/Developer/ShaderFormatVectorVM/Private/VectorVMBackend.h
@@ -71,6 +71,7 @@ enum class ECallScalarizeMode : uint8
 ECallScalarizeMode get_scalarize_mode(ir_function_signature* in_sig);
 EVectorVMOp get_special_vm_opcode(ir_function_signature* signature);
 
+void vm_remove_self_ref_assignments(exec_list* instructions, _mesa_glsl_parse_state *state);
 void vm_matrices_to_vectors(exec_list* instructions, _mesa_glsl_parse_state *state);
 bool do_vec_op_to_scalar(exec_list *instructions, _mesa_glsl_parse_state* parse_state);
 bool vm_flatten_branches_to_selects(exec_list *instructions, _mesa_glsl_parse_state *state);
diff --git a/Engine/Source/Developer/ShaderFormatVectorVM/Private/ir_vm_propagate_non_expressions_visitor.cpp b/Engine/Source/Developer/ShaderFormatVectorVM/Private/ir_vm_propagate_non_expressions_visitor.cpp
--- a/Engine/Source/Developer/ShaderFormatVectorVM/Private/ir_vm_propagate_non_expressions_visitor.cpp
+++ b/Engine/Source/Developer/ShaderFormatVectorVM/Private/ir_vm_propagate_non_expressions_visitor.cpp
@@ -419,3 +419,182 @@ void vm_matrices_to_vectors(exec_list* ir, _mesa_glsl_parse_state* state)
 {	
 	ir_matrices_to_vectors::run(ir, state);
 }
+
+//////////////////////////////////////////////////////////////////////////
+
+
+/** Replaces any assignments that reference themselves with ones that assign to a new temporary. It is potentially unsafe for the VM to read and write to the same location. */
+class ir_rem_self_ref_assignments_visitor final : public ir_rvalue_visitor
+{
+	_mesa_glsl_parse_state *parse_state;
+
+	struct replace_info
+	{
+		/** Component of the variable to replace that should be replaced. */
+		unsigned component;
+
+		/** new temp variable to replace with. */
+		ir_variable* replacement;
+	};
+
+	/** Map of variables we should replace in subsequent rvalues. */
+	TMap<ir_variable*, replace_info> to_replace;
+
+	/* The current assignment we're in. */
+	ir_assignment* curr_assign;
+
+	ir_variable* just_replaced_var;
+	unsigned just_replaced_comp;
+
+	bool progress;
+
+public:
+	ir_rem_self_ref_assignments_visitor(_mesa_glsl_parse_state *in_state)
+	{
+		parse_state = in_state;
+		progress = false;
+
+		curr_assign = nullptr;;
+
+		just_replaced_var = nullptr;
+		just_replaced_comp = 0;
+	}
+
+	virtual ~ir_rem_self_ref_assignments_visitor()
+	{
+	}
+
+	unsigned get_component_from_matrix_array_deref(ir_dereference_array* array_deref)
+	{
+		check(array_deref);
+		check(array_deref->variable_referenced()->type->is_matrix());
+		ir_constant* index = array_deref->array_index->as_constant();
+		check((index->type == glsl_type::uint_type || index->type == glsl_type::int_type) && index->type->is_scalar());
+		unsigned deref_idx = index->type == glsl_type::uint_type ? index->value.u[0] : index->value.i[0];
+		return deref_idx * array_deref->variable_referenced()->type->vector_elements;
+	}
+
+	unsigned get_dest_comp(unsigned write_mask)
+	{
+		if (write_mask == (1 << 0)) return 0;
+		if (write_mask == (1 << 1)) return 1;
+		if (write_mask == (1 << 2)) return 2;
+		if (write_mask == (1 << 3)) return 3;
+		check(0);
+		return 0;
+	}
+
+	virtual void handle_rvalue(ir_rvalue **rvalue)
+	{
+		if (curr_assign && rvalue && *rvalue && !in_assignee)
+		{
+			ir_dereference* deref = (*rvalue)->as_dereference();
+			ir_dereference_array* array_deref = (*rvalue)->as_dereference_array();
+			ir_swizzle* swiz = (*rvalue)->as_swizzle();
+
+			ir_variable* search_var = (*rvalue)->variable_referenced();
+			unsigned search_comp = 0;
+			if (swiz)
+			{
+				if (ir_dereference_array* swiz_array_deref = swiz->val->as_dereference_array())
+				{
+					search_comp = get_component_from_matrix_array_deref(swiz_array_deref);
+				}
+				search_comp += swiz->mask.x;
+			}
+			else if (array_deref)
+			{
+				//We can only handle matrix array derefs but these will have an outer swizzle that we'll work with. 
+				check(array_deref->array->type->is_matrix());
+				search_var = nullptr;
+			}
+			else if (!deref || !deref->type->is_scalar())
+			{
+				//If we're not a deref or we're not a straight scalar deref then we should leave this alone.
+				search_var = nullptr;
+			}
+
+			if (search_var == just_replaced_var && search_comp == just_replaced_comp)
+			{
+				//We've just created a replacement for these vars in this assignment.
+				//This deref can must remain in place but future ones must be replaced.
+				return;
+			}
+			
+			replace_info* info = to_replace.Find(search_var);
+
+			unsigned assign_comp = get_dest_comp(curr_assign->write_mask);
+			if (curr_assign->lhs->variable_referenced() == search_var && search_comp == assign_comp)
+			{
+				//We're in an assignment that is reading from the same component it's writing to.
+				if (info == nullptr)
+				{
+					//Create a new temporary for the assignment and add it to the replacement map so that future refs to it are replaced.
+					ir_variable* new_var = new(parse_state) ir_variable(search_var->type->get_base_type(), "self_assign_replacement", ir_var_temporary);
+					replace_info& new_info = to_replace.Add(search_var);
+					new_info.component = assign_comp;
+					new_info.replacement = new_var;
+
+					just_replaced_var = search_var;
+					just_replaced_comp = search_comp;
+
+					curr_assign->insert_before(new_var);
+					curr_assign->set_lhs(new(parse_state) ir_dereference_variable(new_var));
+					progress = true;
+				}
+			}
+			else
+			{
+				//Look to see if this var deref needs to be replaced.
+				if (info && info->component == search_comp)
+				{
+					ir_rvalue* new_rval = new(parse_state) ir_dereference_variable(info->replacement);
+					(*rvalue) = new_rval;
+					progress = true;
+				}
+			}
+		}
+	}
+
+	virtual ir_visitor_status visit_enter(ir_assignment *assign)
+	{
+		if (assign->condition)
+		{
+			_mesa_glsl_error(parse_state, "conditional assignment in instruction stream");
+			return visit_stop;
+		}
+
+		curr_assign = assign;
+		just_replaced_var = nullptr;
+		just_replaced_comp = 0;
+		return visit_continue;
+	}
+
+	virtual ir_visitor_status visit_leave(ir_assignment *assign)
+	{
+		curr_assign = nullptr;
+		just_replaced_var = nullptr;
+		just_replaced_comp = 0;
+		return visit_continue;
+	}
+
+	static void run(exec_list *ir, _mesa_glsl_parse_state *state)
+	{
+		bool progress = false;
+		do
+		{
+			ir_rem_self_ref_assignments_visitor vis(state);
+			visit_list_elements(&vis, ir);
+
+			progress = vis.progress;
+
+			progress = do_dead_code(ir, false) || progress;
+			progress = do_dead_code_local(ir) || progress;
+		} while (progress);
+	}
+};
+
+void vm_remove_self_ref_assignments(exec_list* ir, _mesa_glsl_parse_state* state)
+{
+	ir_rem_self_ref_assignments_visitor::run(ir, state);
+}
\ No newline at end of file
