Fix static analysis warning
#rnx

diff --git a/Engine/Source/Developer/ShaderFormatVectorVM/Private/ir_vm_propagate_non_expressions_visitor.cpp b/Engine/Source/Developer/ShaderFormatVectorVM/Private/ir_vm_propagate_non_expressions_visitor.cpp
--- a/Engine/Source/Developer/ShaderFormatVectorVM/Private/ir_vm_propagate_non_expressions_visitor.cpp
+++ b/Engine/Source/Developer/ShaderFormatVectorVM/Private/ir_vm_propagate_non_expressions_visitor.cpp
@@ -529,6 +529,8 @@ public:
 				//We're in an assignment that is reading from the same component it's writing to.
 				if (info == nullptr)
 				{
+					checkSlow(search_var);
+					
 					//Create a new temporary for the assignment and add it to the replacement map so that future refs to it are replaced.
 					ir_variable* new_var = new(parse_state) ir_variable(search_var->type->get_base_type(), "self_assign_replacement", ir_var_temporary);
 					replace_info& new_info = to_replace.Add(search_var);
