Fixed call to FString::Printf with C-string argument.

#rnx

diff --git a/Engine/Source/Developer/ShaderFormatVectorVM/Private/ir_vm_gen_bytecode_visitor.cpp b/Engine/Source/Developer/ShaderFormatVectorVM/Private/ir_vm_gen_bytecode_visitor.cpp
--- a/Engine/Source/Developer/ShaderFormatVectorVM/Private/ir_vm_gen_bytecode_visitor.cpp
+++ b/Engine/Source/Developer/ShaderFormatVectorVM/Private/ir_vm_gen_bytecode_visitor.cpp
@@ -1092,7 +1092,7 @@ struct op_external_func : public op_base
 	
 	virtual FString to_string() 
 	{ 
-		FString Str = FString::Printf(TEXT("%S[func%d]("), sig->function_name(), function_index);
+		FString Str = FString::Printf(TEXT("%s[func%d]("), ANSI_TO_TCHAR(sig->function_name()), function_index);
 		int32 param_idx = 0;
 		FString ComponentStr;
 		for (int32 i = 0; i < inputs.Num(); ++i, ++param_idx)
@@ -1676,7 +1676,7 @@ class ir_gen_vvm_visitor : public ir_hierarchical_visitor
 			}
 		}
 
-		FString ExprString = FString::Printf(TEXT("%S %s"), expression->type->name, ir_expression::operator_string(expression->operation));
+		FString ExprString = FString::Printf(TEXT("%s %s"), ANSI_TO_TCHAR(expression->type->name), ir_expression::operator_string(expression->operation));
 		ExprString += TEXT("(");
 		for (unsigned i = 0; i < expression->get_num_operands(); ++i)
 		{
