Fix for vm compiler bug introduced in 9743796 that would cause more temporaries to be used than needed.

diff --git a/Engine/Plugins/FX/Niagara/Source/NiagaraCore/Private/NiagaraCustomVersion.cpp b/Engine/Plugins/FX/Niagara/Source/NiagaraCore/Private/NiagaraCustomVersion.cpp
--- a/Engine/Plugins/FX/Niagara/Source/NiagaraCore/Private/NiagaraCustomVersion.cpp
+++ b/Engine/Plugins/FX/Niagara/Source/NiagaraCore/Private/NiagaraCustomVersion.cpp
@@ -8,4 +8,4 @@ const FGuid FNiagaraCustomVersion::GUID(0xFCF57AFA, 0x50764283, 0xB9A9E658, 0xFF
 // Register the custom version with core
 FCustomVersionRegistration GRegisterNiagaraCustomVersion(FNiagaraCustomVersion::GUID, FNiagaraCustomVersion::LatestVersion, TEXT("NiagaraVer"));
 
-const FGuid FNiagaraCustomVersion::LatestScriptCompileVersion(0x1E9267F8, 0xEA19004A, 0xBF0BA00E, 0x9231686F);
\ No newline at end of file
+const FGuid FNiagaraCustomVersion::LatestScriptCompileVersion(0x0F93E860, 0x6A30224B, 0xA9000879, 0xA7693A89);
\ No newline at end of file
diff --git a/Engine/Source/Developer/ShaderFormatVectorVM/Private/ir_vm_gen_bytecode_visitor.cpp b/Engine/Source/Developer/ShaderFormatVectorVM/Private/ir_vm_gen_bytecode_visitor.cpp
--- a/Engine/Source/Developer/ShaderFormatVectorVM/Private/ir_vm_gen_bytecode_visitor.cpp
+++ b/Engine/Source/Developer/ShaderFormatVectorVM/Private/ir_vm_gen_bytecode_visitor.cpp
@@ -628,12 +628,16 @@ struct op_base
 						}
 					}
 
-					component->offset = temp_registers.Num();
-					CompilationOutput.MaxTempRegistersUsed = FMath::Max(temp_registers.Num(), CompilationOutput.MaxTempRegistersUsed);
-					temp_registers.Add(component->last_read);
+					//If we didn't find a free temp register to use, add a new one.
+					if (component->offset == INDEX_NONE)
+					{
+						component->offset = temp_registers.Num();
+						CompilationOutput.MaxTempRegistersUsed = FMath::Max(temp_registers.Num(), CompilationOutput.MaxTempRegistersUsed);
+						temp_registers.Add(component->last_read);
 #if VM_VERBOSE_LOGGING
-					UE_LOG(LogVVMBackend, Log, TEXT("OP:%d | Comp:%p allocated Reg: %d | Last Read: %d |"), op_idx, component, j, component->last_read);
+						UE_LOG(LogVVMBackend, Log, TEXT("OP:%d | Comp:%p allocated Reg: %d | Last Read: %d |"), op_idx, component, j, component->last_read);
 #endif
+					}
  
 // 					if (component->offset == INDEX_NONE)
 // 					{
