Fixed incorrect dissassembly when WITH_EDITOR is not defined.

The opcode was only added to the argument list when WITH_EDITOR was defined, so all the positional arguments used by the string formatter were offset by -1.

diff --git a/Engine/Source/Developer/ShaderFormatVectorVM/Private/ir_vm_gen_bytecode_visitor.cpp b/Engine/Source/Developer/ShaderFormatVectorVM/Private/ir_vm_gen_bytecode_visitor.cpp
--- a/Engine/Source/Developer/ShaderFormatVectorVM/Private/ir_vm_gen_bytecode_visitor.cpp
+++ b/Engine/Source/Developer/ShaderFormatVectorVM/Private/ir_vm_gen_bytecode_visitor.cpp
@@ -687,6 +687,8 @@ struct op_hook : public op_base
 
 #if WITH_EDITOR
 		Args.Add(VectorVM::GetOpName(op_code));
+#else
+		Args.Add(FString::Printf(TEXT("OP_%d"), op_code));
 #endif
 
 		if(source_component)
@@ -751,6 +753,8 @@ struct op_standard : public op_base
 
 #if WITH_EDITOR
 		Args.Add(VectorVM::GetOpName(op_code));
+#else
+		Args.Add(FString::Printf(TEXT("OP_%d"), op_code));
 #endif
 
 		Args.Add(get_location_string(dest_component->owner->location));
