Fixed crashes in the vector VM backend of the shader compiler when calling void functions or functions with output parameters.

The branch flattening code assumed the function will always have a return value. Also, it never advanced the current parameter pointer when going through the function arguments, so it was trying to dereference the first parameter for all the arguments.

#rnx

diff --git a/Engine/Source/Developer/ShaderFormatVectorVM/Private/ir_vm_flatten_branches_to_selects_visitor.cpp b/Engine/Source/Developer/ShaderFormatVectorVM/Private/ir_vm_flatten_branches_to_selects_visitor.cpp
--- a/Engine/Source/Developer/ShaderFormatVectorVM/Private/ir_vm_flatten_branches_to_selects_visitor.cpp
+++ b/Engine/Source/Developer/ShaderFormatVectorVM/Private/ir_vm_flatten_branches_to_selects_visitor.cpp
@@ -242,11 +242,13 @@ public:
 			check(curr_assignments);
 
 			//void* parent = ralloc_parent(call);
-			ir_dereference_variable* new_deref = replace_assigned_val_with_temp(call->return_deref);
-
-			ir_assignment* new_assign = new(parse_state)ir_assignment(call->return_deref, new_deref);
-			call->return_deref = new_deref;
-			curr_assignments->Add(new_assign);
+			if (call->return_deref)
+			{
+				ir_dereference_variable* new_deref = replace_assigned_val_with_temp(call->return_deref);
+				ir_assignment* new_assign = new(parse_state)ir_assignment(call->return_deref, new_deref);
+				call->return_deref = new_deref;
+				curr_assignments->Add(new_assign);
+			}
 
 			ir_rvalue* actual_param = (ir_rvalue*)call->actual_parameters.get_head();
 			foreach_iter(exec_list_iterator, iter, call->callee->parameters)
@@ -254,12 +256,13 @@ public:
 				ir_variable *var = (ir_variable *)iter.get();
 				if (var->mode == ir_var_out || var->mode == ir_var_inout)
 				{
-					new_deref = replace_assigned_val_with_temp(actual_param->as_dereference());
-					new_assign = new(parse_state)ir_assignment(actual_param, new_deref);
+					ir_dereference_variable* new_deref = replace_assigned_val_with_temp(actual_param->as_dereference());
+					ir_assignment* new_assign = new(parse_state)ir_assignment(actual_param, new_deref);
 					check(actual_param->next && actual_param->prev);
 					actual_param->replace_with(new_deref);
 					curr_assignments->Add(new_assign);
 				}
+				actual_param = (ir_rvalue*)actual_param->next;
 			}
 		}
 
