Fixed crash when using a loop with more than 32 iterations in custom HLSL.

The default threshold for loop unrolling was 32, so the compiler tried to make a dynamic loop when we had more iterations. The vector VM doesn't support that, so the codegen was never designed to understand dynamic loops, leading to an assertion failure. Increased the threshold to a large value, so we never hit this problem in practice.

#rnx

diff --git a/Engine/Source/Developer/ShaderFormatVectorVM/Private/VectorVMBackend.cpp b/Engine/Source/Developer/ShaderFormatVectorVM/Private/VectorVMBackend.cpp
--- a/Engine/Source/Developer/ShaderFormatVectorVM/Private/VectorVMBackend.cpp
+++ b/Engine/Source/Developer/ShaderFormatVectorVM/Private/VectorVMBackend.cpp
@@ -180,6 +180,10 @@ char* FVectorVMCodeBackend::GenerateCode(exec_list* ir, _mesa_glsl_parse_state*
 
 	FlattenUniformBufferStructures(ir, state);
 
+	// The vector VM doesn't support dynamic loops, and the codegen panics if it sees loops. We'll just set the unroll threshold
+	// to a very large value, to make sure loops are always unrolled.
+	state->maxunrollcount = 16384;
+
 	state->conservative_propagation = false;
 	{
 		SCOPE_CYCLE_COUNTER(STAT_VVMInitMisc);
