Fix for VM compile issue with matrix constants.

diff --git a/Engine/Source/Developer/ShaderFormatVectorVM/Private/ir_vm_propagate_non_expressions_visitor.cpp b/Engine/Source/Developer/ShaderFormatVectorVM/Private/ir_vm_propagate_non_expressions_visitor.cpp
--- a/Engine/Source/Developer/ShaderFormatVectorVM/Private/ir_vm_propagate_non_expressions_visitor.cpp
+++ b/Engine/Source/Developer/ShaderFormatVectorVM/Private/ir_vm_propagate_non_expressions_visitor.cpp
@@ -339,31 +339,53 @@ public:
 			}
 			else if (ir_dereference_variable* matrix_deref = assign->lhs->as_dereference_variable())
 			{
-				ir_dereference_variable* src_mat = assign->rhs->as_dereference_variable();
-				check(src_mat && src_mat->type->is_matrix());
-				ir_dereference* src_derefs[4];
-				//void* p = ralloc_parent(array_deref);
-				if (src_mat->variable_referenced()->mode == ir_var_uniform)
+				if (ir_dereference_variable* src_mat = assign->rhs->as_dereference_variable())
 				{
-					src_derefs[0] = new(parse_state) ir_dereference_array(src_mat->variable_referenced(), new(parse_state) ir_constant(0));
-					src_derefs[1] = new(parse_state) ir_dereference_array(src_mat->variable_referenced(), new(parse_state) ir_constant(1));
-					src_derefs[2] = new(parse_state) ir_dereference_array(src_mat->variable_referenced(), new(parse_state) ir_constant(2));
-					src_derefs[3] = new(parse_state) ir_dereference_array(src_mat->variable_referenced(), new(parse_state) ir_constant(3));
+					check(src_mat->type->is_matrix());
+					ir_dereference* src_derefs[4];
+					//void* p = ralloc_parent(array_deref);
+					if (src_mat->variable_referenced()->mode == ir_var_uniform)
+					{
+						src_derefs[0] = new(parse_state) ir_dereference_array(src_mat->variable_referenced(), new(parse_state) ir_constant(0));
+						src_derefs[1] = new(parse_state) ir_dereference_array(src_mat->variable_referenced(), new(parse_state) ir_constant(1));
+						src_derefs[2] = new(parse_state) ir_dereference_array(src_mat->variable_referenced(), new(parse_state) ir_constant(2));
+						src_derefs[3] = new(parse_state) ir_dereference_array(src_mat->variable_referenced(), new(parse_state) ir_constant(3));
 
+					}
+					else
+					{
+						MatrixVectors& src_mv = MatrixVectorMap.FindChecked(src_mat->variable_referenced());
+						src_derefs[0] = new(parse_state) ir_dereference_variable(src_mv.v[0]);
+						src_derefs[1] = new(parse_state) ir_dereference_variable(src_mv.v[1]);
+						src_derefs[2] = new(parse_state) ir_dereference_variable(src_mv.v[2]);
+						src_derefs[3] = new(parse_state) ir_dereference_variable(src_mv.v[3]);
+					}
+
+					assign->insert_before(new(parse_state) ir_assignment(new(parse_state) ir_dereference_variable(mv.v[0]), src_derefs[0], assign->condition, 0xF));
+					assign->insert_before(new(parse_state) ir_assignment(new(parse_state) ir_dereference_variable(mv.v[1]), src_derefs[1], assign->condition, 0xF));
+					assign->insert_before(new(parse_state) ir_assignment(new(parse_state) ir_dereference_variable(mv.v[2]), src_derefs[2], assign->condition, 0xF));
+					assign->insert_before(new(parse_state) ir_assignment(new(parse_state) ir_dereference_variable(mv.v[3]), src_derefs[3], assign->condition, 0xF));
 				}
 				else
 				{
-					MatrixVectors& src_mv = MatrixVectorMap.FindChecked(src_mat->variable_referenced());
-					src_derefs[0] = new(parse_state) ir_dereference_variable(src_mv.v[0]);
-					src_derefs[1] = new(parse_state) ir_dereference_variable(src_mv.v[1]);
-					src_derefs[2] = new(parse_state) ir_dereference_variable(src_mv.v[2]);
-					src_derefs[3] = new(parse_state) ir_dereference_variable(src_mv.v[3]);
+					//Pull the vectors out of the matrix into the new vector temporaries.
+					ir_constant* mat_const = assign->rhs->as_constant();
+					check(mat_const);
+
+					const glsl_type* mat_type = mat_const->type;					
+					for(unsigned mat_col = 0; mat_col < mat_type->matrix_columns; ++mat_col)
+					{
+						ir_constant_data new_vec_data;
+						memset(&new_vec_data, 0, sizeof(new_vec_data));
+						for (unsigned vec_elem = 0; vec_elem < mat_type->vector_elements; ++vec_elem)
+						{
+							new_vec_data.f[vec_elem] = mat_const->value.f[mat_col * mat_type->vector_elements + vec_elem];
+						}
+						ir_constant* new_vec_const = new(parse_state) ir_constant(mat_type->column_type(), &new_vec_data);
+						assign->insert_before(new(parse_state) ir_assignment(new(parse_state) ir_dereference_variable(mv.v[mat_col]), new_vec_const, assign->condition, 0xF));
+					}
 				}
 
-				assign->insert_before(new(parse_state) ir_assignment(new(parse_state) ir_dereference_variable(mv.v[0]), src_derefs[0], assign->condition, 0xF));
-				assign->insert_before(new(parse_state) ir_assignment(new(parse_state) ir_dereference_variable(mv.v[1]), src_derefs[1], assign->condition, 0xF));
-				assign->insert_before(new(parse_state) ir_assignment(new(parse_state) ir_dereference_variable(mv.v[2]), src_derefs[2], assign->condition, 0xF));
-				assign->insert_before(new(parse_state) ir_assignment(new(parse_state) ir_dereference_variable(mv.v[3]), src_derefs[3], assign->condition, 0xF));
 				assign->remove();
 			}
 			else
