Fix for swizzled non scalar constants in the scalarization visitor.

diff --git a/Engine/Source/Developer/ShaderFormatVectorVM/Private/ir_vm_scalarize_visitor.cpp b/Engine/Source/Developer/ShaderFormatVectorVM/Private/ir_vm_scalarize_visitor.cpp
--- a/Engine/Source/Developer/ShaderFormatVectorVM/Private/ir_vm_scalarize_visitor.cpp
+++ b/Engine/Source/Developer/ShaderFormatVectorVM/Private/ir_vm_scalarize_visitor.cpp
@@ -468,11 +468,25 @@ class ir_scalarize_visitor2 : public ir_hierarchical_visitor
 				break;
 			}
 
-
 			swiz->mask.num_components = 1;
-			swiz->mask.x = src_comp;
 			swiz->mask.has_duplicates = false;
 			swiz->type = swiz->type->get_base_type();
+
+			if (ir_constant* const_val = swiz->val->as_constant())
+			{
+				//The scalar itself may be scalar but it's internals could be a non scalar constant.
+				//In that case we visit the constant and it'll convert itself to the right scalar based on dest_componnent.
+				unsigned old_dest = dest_component;
+				dest_component = src_comp;
+				swiz->val->accept(this);
+				dest_component = old_dest;
+
+				swiz->mask.x = 0;
+			}
+			else
+			{
+				swiz->mask.x = src_comp;
+			}
 		}
 
 		return visit_continue_with_parent;
