Fix for VM bytecode generation which was failing in a very small number of cases of the form:
A = GetFirstValue();
B = A;
A = GetSecondValue();
C = A; // <-- this would get GetSecondValue() assigned to C instead of GetFirstValue()

Adds an SSA pass to help reduce overwriting registers, adds logic to remove self assignments, and adds logic to stop propagating when values are dirtied

From a sampling of scripts, about 2% are changed (where 1-3 unnecessary ops have been deleted).

diff --git a/Engine/Source/Developer/ShaderFormatVectorVM/Private/ir_vm_propagate_non_expressions_visitor.cpp b/Engine/Source/Developer/ShaderFormatVectorVM/Private/ir_vm_propagate_non_expressions_visitor.cpp
--- a/Engine/Source/Developer/ShaderFormatVectorVM/Private/ir_vm_propagate_non_expressions_visitor.cpp
+++ b/Engine/Source/Developer/ShaderFormatVectorVM/Private/ir_vm_propagate_non_expressions_visitor.cpp
@@ -22,6 +22,147 @@ PRAGMA_ENABLE_SHADOW_VARIABLE_WARNINGS
 
 #include "VectorVM.h"
 
+namespace ir_propagate_helper
+{
+
+static unsigned get_component_from_matrix_array_deref(ir_dereference_array* array_deref)
+{
+	check(array_deref);
+	check(array_deref->variable_referenced()->type->is_matrix());
+	ir_constant* index = array_deref->array_index->as_constant();
+	check((index->type == glsl_type::uint_type || index->type == glsl_type::int_type) && index->type->is_scalar());
+	unsigned deref_idx = index->type == glsl_type::uint_type ? index->value.u[0] : index->value.i[0];
+	return deref_idx * array_deref->variable_referenced()->type->vector_elements;
+}
+
+static ir_variable* get_rvalue_variable(ir_rvalue* rvalue, unsigned int& search_comp)
+{
+	ir_dereference* deref = rvalue->as_dereference();
+	ir_dereference_array* array_deref = rvalue->as_dereference_array();
+	ir_swizzle* swiz = rvalue->as_swizzle();
+
+	ir_variable* search_var = rvalue->variable_referenced();
+	if (swiz)
+	{
+		if (ir_dereference_array* swiz_array_deref = swiz->val->as_dereference_array())
+		{
+			search_comp = get_component_from_matrix_array_deref(swiz_array_deref);
+		}
+		search_comp += swiz->mask.x;
+	}
+	else if (array_deref)
+	{
+		//We can only handle matrix array derefs but these will have an outer swizzle that we'll work with. 
+		check(array_deref->array->type->is_matrix());
+		search_var = nullptr;
+	}
+	else if (!deref || !deref->type->is_scalar())
+	{
+		//If we're not a deref or we're not a straight scalar deref then we should leave this alone.
+		search_var = nullptr;
+	}
+
+	return search_var;
+}
+
+};
+
+
+// before we worry about propagating expressions we want to clean up our existing assignments
+// so that we get rid of any redundant assignments (A = A) as well as move towards an SSA
+// form so we avoid 'A = B; A = C' (instead prefering A = B; D = C).  This will allow our existing
+// propagation visitors to work.
+class ir_ssa_visitor final : public ir_rvalue_visitor
+{
+public:
+		_mesa_glsl_parse_state* parse_state;
+	bool progress = false;
+
+	ir_ssa_visitor(_mesa_glsl_parse_state* in_state)
+	{
+		parse_state = in_state;
+	}
+
+	using ScalarVariable = TTuple<ir_variable*, int32>;
+	TMap<ScalarVariable, ir_rvalue*> ReplacementMap;
+
+	virtual void handle_rvalue(ir_rvalue** rvalue) override
+	{
+		if (rvalue && *rvalue)
+		{
+			uint32 ComponentIndex = 0;
+			if (ir_variable* Variable = ir_propagate_helper::get_rvalue_variable(*rvalue, ComponentIndex))
+			{
+				if (ir_rvalue** Replacement = ReplacementMap.Find(MakeTuple(Variable, ComponentIndex)))
+				{
+					if (*Replacement)
+					{
+						ir_rvalue* NewRValue = (*Replacement)->clone(parse_state, nullptr);
+						*rvalue = NewRValue;
+						progress = true;
+					}
+				}
+			}
+		}
+	}
+
+	virtual ir_visitor_status visit_leave(ir_assignment* assign) override
+	{
+		ir_variable* lhs = assign->lhs->variable_referenced();
+
+		const int32 ComponentIndex = FMath::CountTrailingZeros(assign->write_mask);
+		check((1 << ComponentIndex) == assign->write_mask);
+
+		// remove redundant assignments
+		uint32 rhs_component_index = 0;
+		if (const ir_variable* rhsVariable = ir_propagate_helper::get_rvalue_variable(assign->rhs, rhs_component_index))
+		{
+			if (lhs == rhsVariable && rhs_component_index == ComponentIndex)
+			{
+				assign->remove();
+				progress = true;
+				return visit_continue;
+			}
+		}
+
+		if (ir_rvalue** Existing = ReplacementMap.Find(MakeTuple(lhs, ComponentIndex)))
+		{
+			// we're already writing to this variable, so we're going to create a new variable for
+			// the LHS, and replace all subsequent matches to our new variable
+			ir_variable* NewVariable = new(parse_state) ir_variable(lhs->type->get_base_type(), "SSA_tmp", ir_var_temporary);
+			ir_rvalue* NewLhs = new(parse_state) ir_dereference_variable(NewVariable);
+
+			base_ir->insert_before(NewVariable);
+			assign->set_lhs(NewLhs);
+			assign->write_mask = 1;
+
+			*Existing = NewLhs;
+			progress = true;
+		}
+		else
+		{
+			ReplacementMap.Add(MakeTuple(lhs, ComponentIndex), nullptr);
+		}
+
+		return ir_rvalue_visitor::visit_leave(assign);
+	}
+
+	static void run(exec_list* ir, _mesa_glsl_parse_state* state)
+	{
+		bool progress = false;
+		do
+		{
+			ir_ssa_visitor ssa_visitor(state);
+			visit_list_elements(&ssa_visitor, ir);
+
+			progress = ssa_visitor.progress;
+
+			progress = do_dead_code(ir, false) || progress;
+			progress = do_dead_code_local(ir) || progress;
+		} while (progress);
+	}
+};
+
 /** Removes any assignments that don't actually map to a VM op but just move some data around. We look for refs and grab the source data direct. */
 class ir_propagate_non_expressions_visitor final : public ir_rvalue_visitor
 {
@@ -40,6 +181,7 @@ class ir_propagate_non_expressions_visitor final : public ir_rvalue_visitor
 	TMap<ir_variable*, var_info> var_info_map;
 
 	TArray<ir_assignment*> assignments;
+	TArray<bool> PreserveAssignment;
 
 	int num_expr;
 	bool progress;
@@ -67,37 +209,6 @@ public:
 		return deref_idx * array_deref->variable_referenced()->type->vector_elements;
 	}
 
-	ir_variable* get_rvalue_variable(ir_rvalue* rvalue, unsigned int& search_comp) const
-	{
-		ir_dereference* deref = rvalue->as_dereference();
-		ir_dereference_array* array_deref = rvalue->as_dereference_array();
-		ir_swizzle* swiz = rvalue->as_swizzle();
-
-		ir_variable* search_var = rvalue->variable_referenced();
-		if (swiz)
-		{
-			if (ir_dereference_array* swiz_array_deref = swiz->val->as_dereference_array())
-			{
-				search_comp = get_component_from_matrix_array_deref(swiz_array_deref);
-			}
-			search_comp += swiz->mask.x;
-		}
-		else if (array_deref)
-		{
-			//We can only handle matrix array derefs but these will have an outer swizzle that we'll work with. 
-			check(array_deref->array->type->is_matrix());
-			search_var = nullptr;
-		}
-		else if (!deref || !deref->type->is_scalar())
-		{
-			//If we're not a deref or we're not a straight scalar deref then we should leave this alone.
-			search_var = nullptr;
-		}
-
-		return search_var;
-	}
-
-
 	virtual void handle_rvalue(ir_rvalue** rvalue)
 	{
 		if (rvalue && *rvalue && !in_assignee)
@@ -105,7 +216,7 @@ public:
 			ir_rvalue** to_replace = rvalue;
 
 			unsigned int search_comp = 0;
-			ir_variable* search_var = get_rvalue_variable(*rvalue, search_comp);
+			ir_variable* search_var = ir_propagate_helper::get_rvalue_variable(*rvalue, search_comp);
 
 			//Search to see if this deref matches any of the non-expression assignments LHS. If so then clone the rhs in it's place.
 
@@ -124,6 +235,10 @@ public:
 					(*rvalue) = new_rval;
 					progress = true;
 				}
+				else if (varinfo->latest_expr_assign[search_comp] != INDEX_NONE)
+				{
+					PreserveAssignment[varinfo->latest_expr_assign[search_comp]] = true;
+				}
 			}
 		}
 	}
@@ -153,6 +268,7 @@ public:
 		var_info& varinfo = var_info_map.FindOrAdd(lhs);
 
 		int32 assign_idx = assignments.Add(assign);
+		bool& Preserve = PreserveAssignment.Add_GetRef(false);
 
 		//Add any new temp or auto assignments. These will be grabbed later to use in replacements in HandleRValue.
 
@@ -162,6 +278,24 @@ public:
 			assign_comp += get_component_from_matrix_array_deref(array_deref);
 		}
 
+		// if our LHS is the RHS of an assignment that is set to copy forward the reference, then we need to invalidate
+		// it so that we don't copy forward a different result
+		for (int32 previous_assign_idx = assign_idx - 1; previous_assign_idx >= 0; --previous_assign_idx)
+		{
+			uint32 rhs_component = 0;
+			ir_variable* rhs_variable = ir_propagate_helper::get_rvalue_variable(assignments[previous_assign_idx]->rhs, rhs_component);
+			if (rhs_variable == lhs && ((1 << rhs_component) & assign->write_mask))
+			{
+				ir_variable* previous_lhs = assignments[previous_assign_idx]->lhs->variable_referenced();
+				if (var_info* previous_varinfo = var_info_map.Find(previous_lhs))
+				{
+					previous_varinfo->latest_expr_assign[rhs_component] = assign_idx;
+					PreserveAssignment[previous_assign_idx] = true;
+					break;
+				}
+			}
+		}
+
 		unsigned write_mask = assign->write_mask;
 		int32 components_written = 0;
 		while (write_mask)
@@ -173,8 +307,7 @@ public:
 				if (num_expr == 0)
 				{
 					unsigned int rhs_component = 0;
-					ir_variable* rhs_variable = get_rvalue_variable(assign->rhs, rhs_component);
-					assign->remove();
+					ir_variable* rhs_variable = ir_propagate_helper::get_rvalue_variable(assign->rhs, rhs_component);
 
 					// handle the case of redundant self assignment
 					if (assign_comp != rhs_component || lhs != rhs_variable)
@@ -186,6 +319,7 @@ public:
 				{
 					check(mode == ir_var_temporary || mode == ir_var_auto);//We can only perform expressions on temp or auto variables.
 					varinfo.latest_expr_assign[assign_comp] = assign_idx;
+					Preserve = true;
 				}
 			}
 			++assign_comp;
@@ -196,14 +330,31 @@ public:
 		return ir_rvalue_visitor::visit_leave(assign);
 	}
 
+	void Finalize()
+	{
+		const int32 AssignmentCount = assignments.Num();
+		for (int32 AssignmentIt = 0; AssignmentIt < AssignmentCount; ++AssignmentIt)
+		{
+			if (!PreserveAssignment[AssignmentIt])
+			{
+				assignments[AssignmentIt]->remove();
+			}
+		}
+	}
+
 	static void run(exec_list* ir, _mesa_glsl_parse_state* state)
 	{
 		bool progress = false;
 		do
 		{
+			ir_ssa_visitor ssa_visitor(state);
+			visit_list_elements(&ssa_visitor, ir);
+
 			ir_propagate_non_expressions_visitor propagate_non_expressions_visitor(state);
 			visit_list_elements(&propagate_non_expressions_visitor, ir);
 
+			propagate_non_expressions_visitor.Finalize();
+
 			progress = propagate_non_expressions_visitor.progress;
 
 			//vm_debug_print("== propagate non expressions - BEFORE DEADCODE ==\n");
diff --git a/Engine/Source/Runtime/Core/Private/UObject/UE5MainStreamObjectVersion.cpp b/Engine/Source/Runtime/Core/Private/UObject/UE5MainStreamObjectVersion.cpp
--- a/Engine/Source/Runtime/Core/Private/UObject/UE5MainStreamObjectVersion.cpp
+++ b/Engine/Source/Runtime/Core/Private/UObject/UE5MainStreamObjectVersion.cpp
@@ -14,7 +14,7 @@ TMap<FGuid, FGuid> FUE5MainStreamObjectVersion::GetSystemGuids()
 	SystemGuids.Add(DevGuids.MATERIALSHADERMAP_DERIVEDDATA_VER, FGuid("65C01D817C9A4EEFAE9E988D41A1F3DD"));
 	SystemGuids.Add(DevGuids.NANITE_DERIVEDDATA_VER, FGuid("7A3EA93546624597BEBE2ADFDE3A7628"));
 	SystemGuids.Add(DevGuids.NIAGARASHADERMAP_DERIVEDDATA_VER, FGuid("8A37C45D24F2423CBE5F8F371DE33575"));
-	SystemGuids.Add(DevGuids.Niagara_LatestScriptCompileVersion, FGuid("C7D06326D0D2074A9B458A402FD9CAC4"));
+	SystemGuids.Add(DevGuids.Niagara_LatestScriptCompileVersion, FGuid("2026F2A19D0C46688E92214FDBAE4CA0"));
 	SystemGuids.Add(DevGuids.SkeletalMeshDerivedDataVersion, FGuid("FB7B2FB546DE49B2A2F88B7F69763CEE"));
 	SystemGuids.Add(DevGuids.STATICMESH_DERIVEDDATA_VER, FGuid("98F7D79A5811013825E75BBCC41ED3E9"));
 	SystemGuids.Add(DevGuids.POSESEARCHDB_DERIVEDDATA_VER, FGuid("4E595C2AC5E947D6BA9ABC874353E5BC"));
