Fixes to niagara VM bytecode generation for if blocks where a variable is assigned and then used within the same block
 if (Condition)
{
  Foo = Bar;	// this would have been replaced with:	branch_flattened_temp0 = Bar;
  Foo2 = Foo;	// this would have been replaced with:	branch_flattened_temp1 = Foo; this should have been replaced by:
				branch_flattened_temp1 = branch_flattened_temp0;
}

diff --git a/Engine/Source/Developer/ShaderFormatVectorVM/Private/ir_vm_flatten_branches_to_selects_visitor.cpp b/Engine/Source/Developer/ShaderFormatVectorVM/Private/ir_vm_flatten_branches_to_selects_visitor.cpp
--- a/Engine/Source/Developer/ShaderFormatVectorVM/Private/ir_vm_flatten_branches_to_selects_visitor.cpp
+++ b/Engine/Source/Developer/ShaderFormatVectorVM/Private/ir_vm_flatten_branches_to_selects_visitor.cpp
@@ -18,9 +18,10 @@ PRAGMA_ENABLE_SHADOW_VARIABLE_WARNINGS
 #include "ir_optimization.h"
 #include "ir_expression_flattening.h"
 #include "ir.h"
+#include "ir_vm_helper.h"
 
 //Helper visitor to replace any flatten branches (currently all) with selection statements the VM can deal with.
-class ir_flatten_branch_to_select_visitor final : public ir_hierarchical_visitor
+class ir_flatten_branch_to_select_visitor final : public ir_rvalue_visitor
 {
 public:
 
@@ -30,6 +31,9 @@ public:
 	TArray<ir_assignment*> a_assignments;
 	TArray<ir_assignment*> b_assignments;
 
+	// keeps track of the rvalues created (branch_flattened) to hold the temporary value calculated within a then/else block
+	TArray<ir_rvalue*> curr_branch_flattened_rvalues;
+
 	TMap<const glsl_type*, ir_function_signature*> select_functions;
 
 	void generate_select_signature(ir_function* func, unsigned base_type, unsigned num_components)
@@ -124,11 +128,15 @@ public:
 		visit_list_elements(this, &curr_if->then_instructions, true);
 		base_ir->insert_before(&curr_if->then_instructions);
 
+		curr_branch_flattened_rvalues.Reset();
+
 		check(b_assignments.Num() == 0);
 		curr_assignments = &b_assignments;
 		visit_list_elements(this, &curr_if->else_instructions, true);
 		base_ir->insert_before(&curr_if->else_instructions);
 
+		curr_branch_flattened_rvalues.Reset();
+
 		TArray<bool> HandledBAssignments;
 		HandledBAssignments.SetNumZeroed(b_assignments.Num());
 
@@ -215,7 +223,49 @@ public:
 			ir_var_temporary);
 		base_ir->insert_before(var);
 
-		return new(parse_state)ir_dereference_variable(var);
+		ir_dereference_variable* new_deref_var = new(parse_state)ir_dereference_variable(var);
+
+		curr_branch_flattened_rvalues.Add(new_deref_var);
+
+		return new_deref_var;
+	}
+
+	// check if the LHS of any of our current assignments are using the current rvalue, if so we'll 
+	// need to replace it with the temporary that we've stashed in curr_branch_flattened_rvalues
+	// Example code that this is trying to resolve:
+	// if (Condition)
+	// {
+	//   Foo = Bar;					// this would have been replaced with:	branch_flattened_temp0 = Bar;
+	//   Foo2 = Foo;				// this would have been replaced with:	branch_flattened_temp1 = Foo; this should have been replaced by:
+	//																		branch_flattened_temp1 = branch_flattened_temp0;
+	// }
+	virtual void handle_rvalue(ir_rvalue** rvalue) override
+	{
+		if (rvalue && *rvalue && curr_if && curr_assignments && !curr_assignments->IsEmpty())
+		{
+			check(curr_branch_flattened_rvalues.Num() == curr_assignments->Num());
+
+			uint32 ComponentIndex = 0;
+			if (ir_variable* Variable = ir_vm_helper::get_rvalue_variable(*rvalue, ComponentIndex))
+			{
+				int32 AssignmentCount = curr_assignments->Num();
+
+				for (int32 AssignIt = 0; AssignIt < AssignmentCount; ++AssignIt)
+				{
+					ir_assignment* prev_assign = (*curr_assignments)[AssignIt];
+
+					uint32 AssignComponentIndex = 0;
+					if (ir_variable* prev_assign_variable = ir_vm_helper::get_rvalue_variable(prev_assign->lhs, AssignComponentIndex))
+					{
+						if (prev_assign_variable == Variable && AssignComponentIndex == ComponentIndex)
+						{
+							ir_rvalue* new_rval = curr_branch_flattened_rvalues[AssignIt]->clone(parse_state, nullptr);
+							(*rvalue) = new_rval;
+						}
+					}
+				}
+			}
+		}
 	}
 
 	ir_visitor_status visit_leave(ir_assignment * assign)
diff --git a/Engine/Source/Developer/ShaderFormatVectorVM/Private/ir_vm_helper.cpp b/Engine/Source/Developer/ShaderFormatVectorVM/Private/ir_vm_helper.cpp
new file mode 100644
--- /dev/null
+++ b/Engine/Source/Developer/ShaderFormatVectorVM/Private/ir_vm_helper.cpp
@@ -0,0 +1,70 @@
+// Copyright Epic Games, Inc. All Rights Reserved.
+
+#include "ir_vm_helper.h"
+
+#include "ShaderFormatVectorVM.h"
+#include "CoreMinimal.h"
+#include "hlslcc.h"
+#include "hlslcc_private.h"
+#include "VectorVMBackend.h"
+
+PRAGMA_DISABLE_SHADOW_VARIABLE_WARNINGS
+#include "glsl_parser_extras.h"
+PRAGMA_ENABLE_SHADOW_VARIABLE_WARNINGS
+
+#include "hash_table.h"
+#include "ir_rvalue_visitor.h"
+#include "PackUniformBuffers.h"
+#include "IRDump.h"
+#include "OptValueNumbering.h"
+#include "ir_optimization.h"
+#include "ir_expression_flattening.h"
+#include "ir.h"
+
+namespace ir_vm_helper
+{
+
+unsigned get_component_from_matrix_array_deref(ir_dereference_array* array_deref)
+{
+	check(array_deref);
+	check(array_deref->variable_referenced()->type->is_matrix());
+	ir_constant* index = array_deref->array_index->as_constant();
+	check((index->type == glsl_type::uint_type || index->type == glsl_type::int_type) && index->type->is_scalar());
+	unsigned deref_idx = index->type == glsl_type::uint_type ? index->value.u[0] : index->value.i[0];
+	return deref_idx * array_deref->variable_referenced()->type->vector_elements;
+}
+
+ir_variable* get_rvalue_variable(ir_rvalue* rvalue, unsigned int& search_comp)
+{
+	ir_dereference* deref = rvalue->as_dereference();
+	ir_dereference_array* array_deref = rvalue->as_dereference_array();
+	ir_swizzle* swiz = rvalue->as_swizzle();
+
+	ir_variable* search_var = rvalue->variable_referenced();
+	search_comp = 0;
+
+	if (swiz)
+	{
+		if (ir_dereference_array* swiz_array_deref = swiz->val->as_dereference_array())
+		{
+			search_comp = get_component_from_matrix_array_deref(swiz_array_deref);
+		}
+		search_comp += swiz->mask.x;
+	}
+	else if (array_deref)
+	{
+		//We can only handle matrix array derefs but these will have an outer swizzle that we'll work with. 
+		check(array_deref->array->type->is_matrix());
+		search_var = nullptr;
+	}
+	else if (!deref || !deref->type->is_scalar())
+	{
+		//If we're not a deref or we're not a straight scalar deref then we should leave this alone.
+		search_var = nullptr;
+	}
+
+	return search_var;
+}
+
+};
+
diff --git a/Engine/Source/Developer/ShaderFormatVectorVM/Private/ir_vm_helper.h b/Engine/Source/Developer/ShaderFormatVectorVM/Private/ir_vm_helper.h
new file mode 100644
--- /dev/null
+++ b/Engine/Source/Developer/ShaderFormatVectorVM/Private/ir_vm_helper.h
@@ -0,0 +1,16 @@
+// Copyright Epic Games, Inc. All Rights Reserved.
+
+#pragma once
+
+class ir_dereference_array;
+class ir_rvalue;
+class ir_variable;
+
+namespace ir_vm_helper
+{
+
+unsigned get_component_from_matrix_array_deref(ir_dereference_array* array_deref);
+ir_variable* get_rvalue_variable(ir_rvalue* rvalue, unsigned int& search_comp);
+
+};
+
diff --git a/Engine/Source/Developer/ShaderFormatVectorVM/Private/ir_vm_propagate_non_expressions_visitor.cpp b/Engine/Source/Developer/ShaderFormatVectorVM/Private/ir_vm_propagate_non_expressions_visitor.cpp
--- a/Engine/Source/Developer/ShaderFormatVectorVM/Private/ir_vm_propagate_non_expressions_visitor.cpp
+++ b/Engine/Source/Developer/ShaderFormatVectorVM/Private/ir_vm_propagate_non_expressions_visitor.cpp
@@ -20,54 +20,9 @@ PRAGMA_ENABLE_SHADOW_VARIABLE_WARNINGS
 #include "ir_expression_flattening.h"
 #include "ir.h"
 
+#include "ir_vm_helper.h"
 #include "VectorVM.h"
 
-namespace ir_propagate_helper
-{
-
-static unsigned get_component_from_matrix_array_deref(ir_dereference_array* array_deref)
-{
-	check(array_deref);
-	check(array_deref->variable_referenced()->type->is_matrix());
-	ir_constant* index = array_deref->array_index->as_constant();
-	check((index->type == glsl_type::uint_type || index->type == glsl_type::int_type) && index->type->is_scalar());
-	unsigned deref_idx = index->type == glsl_type::uint_type ? index->value.u[0] : index->value.i[0];
-	return deref_idx * array_deref->variable_referenced()->type->vector_elements;
-}
-
-static ir_variable* get_rvalue_variable(ir_rvalue* rvalue, unsigned int& search_comp)
-{
-	ir_dereference* deref = rvalue->as_dereference();
-	ir_dereference_array* array_deref = rvalue->as_dereference_array();
-	ir_swizzle* swiz = rvalue->as_swizzle();
-
-	ir_variable* search_var = rvalue->variable_referenced();
-	if (swiz)
-	{
-		if (ir_dereference_array* swiz_array_deref = swiz->val->as_dereference_array())
-		{
-			search_comp = get_component_from_matrix_array_deref(swiz_array_deref);
-		}
-		search_comp += swiz->mask.x;
-	}
-	else if (array_deref)
-	{
-		//We can only handle matrix array derefs but these will have an outer swizzle that we'll work with. 
-		check(array_deref->array->type->is_matrix());
-		search_var = nullptr;
-	}
-	else if (!deref || !deref->type->is_scalar())
-	{
-		//If we're not a deref or we're not a straight scalar deref then we should leave this alone.
-		search_var = nullptr;
-	}
-
-	return search_var;
-}
-
-};
-
-
 // before we worry about propagating expressions we want to clean up our existing assignments
 // so that we get rid of any redundant assignments (A = A) as well as move towards an SSA
 // form so we avoid 'A = B; A = C' (instead prefering A = B; D = C).  This will allow our existing
@@ -91,7 +46,7 @@ public:
 		if (rvalue && *rvalue)
 		{
 			uint32 ComponentIndex = 0;
-			if (ir_variable* Variable = ir_propagate_helper::get_rvalue_variable(*rvalue, ComponentIndex))
+			if (ir_variable* Variable = ir_vm_helper::get_rvalue_variable(*rvalue, ComponentIndex))
 			{
 				if (ir_rvalue** Replacement = ReplacementMap.Find(MakeTuple(Variable, ComponentIndex)))
 				{
@@ -115,7 +70,7 @@ public:
 
 		// remove redundant assignments
 		uint32 rhs_component_index = 0;
-		if (const ir_variable* rhsVariable = ir_propagate_helper::get_rvalue_variable(assign->rhs, rhs_component_index))
+		if (const ir_variable* rhsVariable = ir_vm_helper::get_rvalue_variable(assign->rhs, rhs_component_index))
 		{
 			if (lhs == rhsVariable && rhs_component_index == ComponentIndex)
 			{
@@ -216,7 +171,7 @@ public:
 			ir_rvalue** to_replace = rvalue;
 
 			unsigned int search_comp = 0;
-			ir_variable* search_var = ir_propagate_helper::get_rvalue_variable(*rvalue, search_comp);
+			ir_variable* search_var = ir_vm_helper::get_rvalue_variable(*rvalue, search_comp);
 
 			//Search to see if this deref matches any of the non-expression assignments LHS. If so then clone the rhs in it's place.
 
@@ -283,7 +238,7 @@ public:
 		for (int32 previous_assign_idx = assign_idx - 1; previous_assign_idx >= 0; --previous_assign_idx)
 		{
 			uint32 rhs_component = 0;
-			ir_variable* rhs_variable = ir_propagate_helper::get_rvalue_variable(assignments[previous_assign_idx]->rhs, rhs_component);
+			ir_variable* rhs_variable = ir_vm_helper::get_rvalue_variable(assignments[previous_assign_idx]->rhs, rhs_component);
 			if (rhs_variable == lhs && ((1 << rhs_component) & assign->write_mask))
 			{
 				ir_variable* previous_lhs = assignments[previous_assign_idx]->lhs->variable_referenced();
@@ -307,7 +262,7 @@ public:
 				if (num_expr == 0)
 				{
 					unsigned int rhs_component = 0;
-					ir_variable* rhs_variable = ir_propagate_helper::get_rvalue_variable(assign->rhs, rhs_component);
+					ir_variable* rhs_variable = ir_vm_helper::get_rvalue_variable(assign->rhs, rhs_component);
 
 					// handle the case of redundant self assignment
 					if (assign_comp != rhs_component || lhs != rhs_variable)
diff --git a/Engine/Source/Runtime/Core/Private/UObject/FortniteMainBranchObjectVersion.cpp b/Engine/Source/Runtime/Core/Private/UObject/FortniteMainBranchObjectVersion.cpp
--- a/Engine/Source/Runtime/Core/Private/UObject/FortniteMainBranchObjectVersion.cpp
+++ b/Engine/Source/Runtime/Core/Private/UObject/FortniteMainBranchObjectVersion.cpp
@@ -11,7 +11,7 @@ TMap<FGuid, FGuid> FFortniteMainBranchObjectVersion::GetSystemGuids()
 	SystemGuids.Add(DevGuids.MATERIALSHADERMAP_DERIVEDDATA_VER, FGuid("1DA85DA5733C4A489A1BE117CDC5ACCC"));
 	SystemGuids.Add(DevGuids.NANITE_DERIVEDDATA_VER, FGuid("93379C2F29624E74AFB14FB13E327FA2"));
 	SystemGuids.Add(DevGuids.NIAGARASHADERMAP_DERIVEDDATA_VER, FGuid("2289D5116CF94BC0AFEAEC541468E645"));
-	SystemGuids.Add(DevGuids.Niagara_LatestScriptCompileVersion, FGuid("7379EA86433D4FD0857C15AFCAB9EC1D"));
+	SystemGuids.Add(DevGuids.Niagara_LatestScriptCompileVersion, FGuid("A731B02963264F99A7CFD5B4AC060E4C"));
 	SystemGuids.Add(DevGuids.SkeletalMeshDerivedDataVersion, FGuid("ACF593EAAE354FCBB34CF44F5AA1BFD2"));
 	SystemGuids.Add(DevGuids.STATICMESH_DERIVEDDATA_VER, FGuid("2B18DC81AAF64015A7901B5A3F8394C0"));
 	
