Fix for mistaken static analysis warning with ir_vm generation

diff --git a/Engine/Source/Developer/ShaderFormatVectorVM/Private/ir_vm_flatten_branches_to_selects_visitor.cpp b/Engine/Source/Developer/ShaderFormatVectorVM/Private/ir_vm_flatten_branches_to_selects_visitor.cpp
--- a/Engine/Source/Developer/ShaderFormatVectorVM/Private/ir_vm_flatten_branches_to_selects_visitor.cpp
+++ b/Engine/Source/Developer/ShaderFormatVectorVM/Private/ir_vm_flatten_branches_to_selects_visitor.cpp
@@ -32,7 +32,9 @@ public:
 	TArray<ir_assignment*> b_assignments;
 
 	// keeps track of the rvalues created (branch_flattened) to hold the temporary value calculated within a then/else block
-	TArray<ir_rvalue*> curr_branch_flattened_rvalues;
+	TArray<ir_rvalue*>* curr_branch_flattened_rvalues;
+	TArray<ir_rvalue*> if_flattened_rvalues;
+	TArray<ir_rvalue*> then_flattened_rvalues;
 
 	TMap<const glsl_type*, ir_function_signature*> select_functions;
 
@@ -53,6 +55,7 @@ public:
 	{
 		curr_if = nullptr;
 		curr_assignments = nullptr;
+		curr_branch_flattened_rvalues = nullptr;
 
 		ir_function* select_func = new(parse_state)ir_function("select");
 		generate_select_signature(select_func, GLSL_TYPE_BOOL, 1);
@@ -124,18 +127,23 @@ public:
 
 		check(a_assignments.Num() == 0);
 		check(curr_assignments == nullptr);
+		check(curr_branch_flattened_rvalues == nullptr);
+		check(if_flattened_rvalues.Num() == 0);
 		curr_assignments = &a_assignments;
+		curr_branch_flattened_rvalues = &if_flattened_rvalues;
 		visit_list_elements(this, &curr_if->then_instructions, true);
 		base_ir->insert_before(&curr_if->then_instructions);
-
-		curr_branch_flattened_rvalues.Reset();
+		if_flattened_rvalues.Reset();
 
 		check(b_assignments.Num() == 0);
+		check(then_flattened_rvalues.Num() == 0);
 		curr_assignments = &b_assignments;
+		curr_branch_flattened_rvalues = &then_flattened_rvalues;
 		visit_list_elements(this, &curr_if->else_instructions, true);
 		base_ir->insert_before(&curr_if->else_instructions);
+		then_flattened_rvalues.Reset();
 
-		curr_branch_flattened_rvalues.Reset();
+		curr_branch_flattened_rvalues = nullptr;
 
 		TArray<bool> HandledBAssignments;
 		HandledBAssignments.SetNumZeroed(b_assignments.Num());
@@ -225,7 +233,7 @@ public:
 
 		ir_dereference_variable* new_deref_var = new(parse_state)ir_dereference_variable(var);
 
-		curr_branch_flattened_rvalues.Add(new_deref_var);
+		curr_branch_flattened_rvalues->Add(new_deref_var);
 
 		return new_deref_var;
 	}
@@ -243,7 +251,7 @@ public:
 	{
 		if (rvalue && *rvalue && curr_if && curr_assignments && !curr_assignments->IsEmpty())
 		{
-			check(curr_branch_flattened_rvalues.Num() == curr_assignments->Num());
+			check(curr_branch_flattened_rvalues->Num() == curr_assignments->Num());
 
 			uint32 ComponentIndex = 0;
 			if (ir_variable* Variable = ir_vm_helper::get_rvalue_variable(*rvalue, ComponentIndex))
@@ -259,7 +267,7 @@ public:
 					{
 						if (prev_assign_variable == Variable && AssignComponentIndex == ComponentIndex)
 						{
-							ir_rvalue* new_rval = curr_branch_flattened_rvalues[AssignIt]->clone(parse_state, nullptr);
+							ir_rvalue* new_rval = (*curr_branch_flattened_rvalues)[AssignIt]->clone(parse_state, nullptr);
 							(*rvalue) = new_rval;
 						}
 					}
