[shaders] add a mutex acquired around all calls into hlslcc as a temporary fix for parallel in-process compilation. also some minor cleanup (include sorting, remove an unused/unimplemented function prototype)

diff --git a/Engine/Source/Developer/ShaderCompilerCommon/Private/CrossCompiler.cpp b/Engine/Source/Developer/ShaderCompilerCommon/Private/CrossCompiler.cpp
new file mode 100644
--- /dev/null
+++ b/Engine/Source/Developer/ShaderCompilerCommon/Private/CrossCompiler.cpp
@@ -0,0 +1,14 @@
+// Copyright Epic Games, Inc. All Rights Reserved.
+
+#include "CrossCompiler.h"
+
+#include "hlslcc.h"
+
+namespace CrossCompiler
+{
+	FCriticalSection* GetCrossCompilerLock()
+	{
+		static FCriticalSection HlslCcCs;
+		return &HlslCcCs;
+	}
+}
\ No newline at end of file
diff --git a/Engine/Source/Developer/ShaderCompilerCommon/Public/CrossCompiler.h b/Engine/Source/Developer/ShaderCompilerCommon/Public/CrossCompiler.h
--- a/Engine/Source/Developer/ShaderCompilerCommon/Public/CrossCompiler.h
+++ b/Engine/Source/Developer/ShaderCompilerCommon/Public/CrossCompiler.h
@@ -5,11 +5,10 @@
 #include "CoreMinimal.h"
 #include "RHIDefinitions.h"
 
+
 // Cross compiler support/common functionality
-extern SHADERCOMPILERCOMMON_API FString CreateCrossCompilerBatchFileContents(
-											const FString& ShaderFile,
-											const FString& OutputFile,
-											const FString& FrequencySwitch,
-											const FString& EntryPoint,
-											const FString& VersionSwitch,
-											const FString& ExtraArguments = TEXT(""));
+namespace CrossCompiler
+{
+extern SHADERCOMPILERCOMMON_API FCriticalSection* GetCrossCompilerLock();
+}
+
diff --git a/Engine/Source/Developer/ShaderFormatOpenGL/Private/OpenGLShaderCompiler.cpp b/Engine/Source/Developer/ShaderFormatOpenGL/Private/OpenGLShaderCompiler.cpp
--- a/Engine/Source/Developer/ShaderFormatOpenGL/Private/OpenGLShaderCompiler.cpp
+++ b/Engine/Source/Developer/ShaderFormatOpenGL/Private/OpenGLShaderCompiler.cpp
@@ -2,14 +2,15 @@
 // ..
 
 #include "CoreMinimal.h"
+#include "CrossCompiler.h"
 #include "HAL/FileManager.h"
+#include "HlslccHeaderWriter.h"
 #include "Misc/FileHelper.h"
 #include "Misc/Paths.h"
 #include "Serialization/MemoryWriter.h"
 #include "ShaderFormatOpenGL.h"
-#include "HlslccHeaderWriter.h"
-#include "SpirvReflectCommon.h"
 #include "ShaderParameterParser.h"
+#include "SpirvReflectCommon.h"
 #include <algorithm>
 #include <regex>
 
@@ -3422,18 +3423,20 @@ void FOpenGLFrontend::CompileShader(const FShaderCompilerInput& Input, FShaderCo
 		bool bDefaultPrecisionIsHalf = (CCFlags & HLSLCC_UseFullPrecisionInPS) == 0;
 		FGlslLanguageSpec* LanguageSpec = CreateLanguageSpec(Version, bDefaultPrecisionIsHalf);
 
-		FHlslCrossCompilerContext CrossCompilerContext(CCFlags, HlslFrequency, HlslCompilerTarget);
-		if (CrossCompilerContext.Init(TCHAR_TO_ANSI(*Input.VirtualSourceFilePath), LanguageSpec))
 		{
-			bCompilationSucceeded = CrossCompilerContext.Run(
-				TCHAR_TO_ANSI(*PreprocessedShader),
-				TCHAR_TO_ANSI(*Input.EntryPointName),
-				BackEnd,
-				&GlslShaderSource,
-				&ErrorLog
-			);
+			FScopeLock HlslCcLock(CrossCompiler::GetCrossCompilerLock());
+			FHlslCrossCompilerContext CrossCompilerContext(CCFlags, HlslFrequency, HlslCompilerTarget);
+			if (CrossCompilerContext.Init(TCHAR_TO_ANSI(*Input.VirtualSourceFilePath), LanguageSpec))
+			{
+				bCompilationSucceeded = CrossCompilerContext.Run(
+					TCHAR_TO_ANSI(*PreprocessedShader),
+					TCHAR_TO_ANSI(*Input.EntryPointName),
+					BackEnd,
+					&GlslShaderSource,
+					&ErrorLog
+				);
+			}
 		}
-
 		delete BackEnd;
 		delete LanguageSpec;
 
diff --git a/Engine/Source/Developer/ShaderFormatVectorVM/Private/VectorVMShaderCompiler.cpp b/Engine/Source/Developer/ShaderFormatVectorVM/Private/VectorVMShaderCompiler.cpp
--- a/Engine/Source/Developer/ShaderFormatVectorVM/Private/VectorVMShaderCompiler.cpp
+++ b/Engine/Source/Developer/ShaderFormatVectorVM/Private/VectorVMShaderCompiler.cpp
@@ -1,17 +1,19 @@
 // Copyright Epic Games, Inc. All Rights Reserved.
 // .
 
-#include "ShaderFormatVectorVM.h"
-#include "VectorVMBackend.h"
 #include "CoreMinimal.h"
-#include "Misc/FileHelper.h"
-#include "ShaderCore.h"
+#include "CrossCompiler.h"
 #include "CrossCompilerCommon.h"
+#include "HAL/CriticalSection.h"
+#include "Misc/FileHelper.h"
+#include "Serialization/MemoryWriter.h"
 #include "ShaderCompilerCommon.h"
+#include "ShaderCore.h"
+#include "ShaderFormatVectorVM.h"
 #include "ShaderPreprocessor.h"
 
 #include "VectorVM.h"
-#include "Serialization/MemoryWriter.h"
+#include "VectorVMBackend.h"
 
 DEFINE_LOG_CATEGORY_STATIC(LogVectorVMShaderCompiler, Log, All); 
 
@@ -136,19 +138,21 @@ bool CompileShader_VectorVM(const FShaderCompilerInput& Input, FShaderCompilerOu
 	FVectorVMCodeBackend VVMBackEnd(CCFlags, HlslCompilerTarget, VMCompilationOutput);
 		FVectorVMLanguageSpec VVMLanguageSpec; 
 
-	bool Result = false;
-	FHlslCrossCompilerContext CrossCompilerContext(CCFlags, Frequency, HlslCompilerTarget);
-	if (CrossCompilerContext.Init(TCHAR_TO_ANSI(*Input.VirtualSourceFilePath), &VVMLanguageSpec))
+	bool bResult = false;
 	{
-		SCOPE_CYCLE_COUNTER(STAT_VectorVM_Compiler_CompileShader_CrossCompilerContextRun);
-
-		Result = CrossCompilerContext.Run(
-			TCHAR_TO_ANSI(*PreprocessedShader),
-			TCHAR_TO_ANSI(*Input.EntryPointName),
-			&VVMBackEnd,
-			&ShaderSource,
-			&ErrorLog
-			) ? 1 : 0;
+		FScopeLock HlslCcLock(CrossCompiler::GetCrossCompilerLock());
+		FHlslCrossCompilerContext CrossCompilerContext(CCFlags, Frequency, HlslCompilerTarget);
+		if (CrossCompilerContext.Init(TCHAR_TO_ANSI(*Input.VirtualSourceFilePath), &VVMLanguageSpec))
+		{
+			SCOPE_CYCLE_COUNTER(STAT_VectorVM_Compiler_CompileShader_CrossCompilerContextRun);
+
+			bResult = CrossCompilerContext.Run(
+				TCHAR_TO_ANSI(*PreprocessedShader),
+				TCHAR_TO_ANSI(*Input.EntryPointName),
+				&VVMBackEnd,
+				&ShaderSource,
+				&ErrorLog);
+		}
 	}
 
 	if (ErrorLog)
@@ -258,6 +262,6 @@ bool CompileShader_VectorVM(const FShaderCompilerInput& Input, FShaderCompilerOu
 		free(ErrorLog);
 	}
 
-	return Result;
+	return bResult;
 }
 
