[shaders] migration of VectorVM shader format to independent preprocessing API (along with some cleanup)

diff --git a/Engine/Plugins/FX/Niagara/Source/NiagaraEditor/Private/NiagaraCompiler.cpp b/Engine/Plugins/FX/Niagara/Source/NiagaraEditor/Private/NiagaraCompiler.cpp
--- a/Engine/Plugins/FX/Niagara/Source/NiagaraEditor/Private/NiagaraCompiler.cpp
+++ b/Engine/Plugins/FX/Niagara/Source/NiagaraEditor/Private/NiagaraCompiler.cpp
@@ -30,7 +30,7 @@
 #include "Serialization/MemoryReader.h"
 #include "ShaderCompiler.h"
 #include "ShaderCore.h"
-#include "ShaderFormatVectorVM.h"
+#include "VectorVMTestCompile.h"
 
 #define LOCTEXT_NAMESPACE "NiagaraCompiler"
 
@@ -1704,10 +1704,9 @@ void FNiagaraEditorModule::TestCompileScriptFromConsole(const TArray<FString>& A
 			Input.Environment.SetDefine(TEXT("AMPLIFICATIONSHADER"), 0);
 			Input.Environment.IncludeVirtualPathToContentsMap.Add(TEXT("/Engine/Generated/NiagaraEmitterInstance.ush"), TranslatedHLSL);
 
-			FShaderCompilerOutput Output;
 			FVectorVMCompilationOutput CompilationOutput;
 			double StartTime = FPlatformTime::Seconds();
-			bool bSucceeded = CompileShader_VectorVM(Input, Output, FString(FPlatformProcess::ShaderDir()), 0, CompilationOutput, GNiagaraSkipVectorVMBackendOptimizations != 0);
+			bool bSucceeded = TestCompileVectorVMShader(Input, FString(FPlatformProcess::ShaderDir()), CompilationOutput, GNiagaraSkipVectorVMBackendOptimizations != 0);
 			float DeltaTime = (float)(FPlatformTime::Seconds() - StartTime);
 
 			if (bSucceeded)
diff --git a/Engine/Source/Developer/ShaderFormatVectorVM/Internal/VectorVMTestCompile.h b/Engine/Source/Developer/ShaderFormatVectorVM/Internal/VectorVMTestCompile.h
new file mode 100644
--- /dev/null
+++ b/Engine/Source/Developer/ShaderFormatVectorVM/Internal/VectorVMTestCompile.h
@@ -0,0 +1,7 @@
+// Copyright Epic Games, Inc. All Rights Reserved.
+
+#pragma once
+
+#include "ShaderFormatVectorVM.h"
+
+extern SHADERFORMATVECTORVM_API bool TestCompileVectorVMShader(const FShaderCompilerInput& Input, const FString& WorkingDirectory, FVectorVMCompilationOutput& VMCompilationOutput, bool bSkipBackendOptimizations);
diff --git a/Engine/Source/Developer/ShaderFormatVectorVM/Private/ShaderFormatVectorVM.cpp b/Engine/Source/Developer/ShaderFormatVectorVM/Private/ShaderFormatVectorVM.cpp
--- a/Engine/Source/Developer/ShaderFormatVectorVM/Private/ShaderFormatVectorVM.cpp
+++ b/Engine/Source/Developer/ShaderFormatVectorVM/Private/ShaderFormatVectorVM.cpp
@@ -8,11 +8,25 @@
 #include "Interfaces/IShaderFormat.h"
 #include "Interfaces/IShaderFormatModule.h"
 #include "hlslcc.h"
+#include "ShaderCompilerCore.h"
 #include "ShaderCore.h"
 
 
-//TODO: Much of the below is some partial work I did towards having the VVM be treat just as a shader platform.
-//This seems like a reasonable way to go but I'm not 100% certain yet.
+extern bool PreprocessVectorVMShader(
+	const FShaderCompilerInput& Input,
+	const FShaderCompilerEnvironment& Environment,
+	class FShaderPreprocessOutput& PreprocessOutput);
+
+extern bool CompileVectorVMShader(
+	const FShaderCompilerInput& Input,
+	const FShaderPreprocessOutput& PreprocessOutput,
+	FShaderCompilerOutput& Output,
+	const FString& WorkingDirectory);
+
+extern void OutputVectorVMDebugData(
+	const FShaderCompilerInput& Input,
+	const FShaderPreprocessOutput& PreprocessOutput,
+	const FShaderCompilerOutput& Output);
 
 static FName NAME_VVM_1_0(TEXT("VVM_1_0"));
 
@@ -49,18 +63,26 @@ public:
 		OutFormats.Add(NAME_VVM_1_0);
 	}
 
-	virtual void CompileShader(FName Format, const struct FShaderCompilerInput& Input, struct FShaderCompilerOutput& Output,const FString& WorkingDirectory) const override
+	virtual bool SupportsIndependentPreprocessing() const override
 	{
-		CheckFormat(Format);
+		return true;
+	}
 
-		if (Format == NAME_VVM_1_0)
-		{
-			CompileShader_VectorVM(Input, Output, WorkingDirectory, (int8)VectorVMFormats::VVM_1_0);
-		}
-		else
-		{
-			check(0);
-		}
+	virtual bool PreprocessShader(const FShaderCompilerInput& Input, const FShaderCompilerEnvironment& Environment, FShaderPreprocessOutput& PreprocessOutput) const override
+	{
+		CheckFormat(Input.ShaderFormat);
+		return PreprocessVectorVMShader(Input, Environment, PreprocessOutput);
+	}
+
+	virtual void CompilePreprocessedShader(const FShaderCompilerInput& Input, const FShaderPreprocessOutput& PreprocessOutput, FShaderCompilerOutput& Output, const FString& WorkingDirectory) const override
+	{
+		CheckFormat(Input.ShaderFormat);
+		CompileVectorVMShader(Input, PreprocessOutput, Output, WorkingDirectory);
+	}
+
+	virtual void OutputDebugData(const FShaderCompilerInput& Input, const FShaderPreprocessOutput& PreprocessOutput, const FShaderCompilerOutput& Output) const override
+	{
+		OutputVectorVMDebugData(Input, PreprocessOutput, Output);
 	}
 
 	virtual const TCHAR* GetPlatformIncludeDirectory() const
diff --git a/Engine/Source/Developer/ShaderFormatVectorVM/Private/VectorVMShaderCompiler.cpp b/Engine/Source/Developer/ShaderFormatVectorVM/Private/VectorVMShaderCompiler.cpp
--- a/Engine/Source/Developer/ShaderFormatVectorVM/Private/VectorVMShaderCompiler.cpp
+++ b/Engine/Source/Developer/ShaderFormatVectorVM/Private/VectorVMShaderCompiler.cpp
@@ -11,10 +11,12 @@
 #include "ShaderCore.h"
 #include "ShaderFormatVectorVM.h"
 #include "ShaderPreprocessor.h"
-#include "Runtime/RenderCore/Internal/ShaderCompilerDefinitions.h"
+#include "ShaderPreprocessTypes.h"
+#include "ShaderCompilerDefinitions.h"
 
 #include "VectorVM.h"
 #include "VectorVMBackend.h"
+#include "VectorVMTestCompile.h"
 
 DEFINE_LOG_CATEGORY_STATIC(LogVectorVMShaderCompiler, Log, All); 
 
@@ -24,122 +26,66 @@ DECLARE_CYCLE_STAT(TEXT("VectorVM - Compiler - CompileShader_VectorVM"), STAT_Ve
 DECLARE_CYCLE_STAT(TEXT("VectorVM - Compiler - PreprocessShader"), STAT_VectorVM_Compiler_CompileShader_VectorVMPreprocessShader, STATGROUP_VectorVM);
 DECLARE_CYCLE_STAT(TEXT("VectorVM - Compiler - CrossCompilerContextRun"), STAT_VectorVM_Compiler_CompileShader_CrossCompilerContextRun, STATGROUP_VectorVM);
 
-
-/**
- * Compile a shader for the VectorVM
- * @param Input - The input shader code and environment.
- * @param Output - Contains shader compilation results upon return.
- */
-bool CompileShader_VectorVM(const FShaderCompilerInput& Input, FShaderCompilerOutput& Output, const FString& WorkingDirectory, uint8 Version)
-{
-	FVectorVMCompilationOutput CompilationOutput;
-	bool Result = CompileShader_VectorVM(Input, Output, WorkingDirectory, Version, CompilationOutput, Input.Environment.CompilerFlags.Contains(CFLAG_SkipOptimizations));
-	
-	if (Result)
-	{
-		FMemoryWriter Ar(Output.ShaderCode.GetWriteAccess(), true);
-		Ar << CompilationOutput;
-		Output.bSucceeded = true;
-	}
-	else if (CompilationOutput.Errors.Len() > 0)
-	{
-		Output.Errors.Add(FShaderCompilerError(*CompilationOutput.Errors));
-	}
-	
-	return Result;
-}
-
-//TODO: Move to this output living in the shader eco-system with the compute shaders too but for now just do things more directly.
-bool CompileShader_VectorVM(const FShaderCompilerInput& Input, FShaderCompilerOutput& Output, const FString& WorkingDirectory, uint8 Version, FVectorVMCompilationOutput& VMCompilationOutput, bool bSkipBackendOptimizations)
+bool PreprocessVectorVMShader(const FShaderCompilerInput& Input, const FShaderCompilerEnvironment& Environment, FShaderPreprocessOutput& Output)
 {
-	SCOPE_CYCLE_COUNTER(STAT_VectorVM_Compiler_CompileShader_VectorVM);
+	SCOPE_CYCLE_COUNTER(STAT_VectorVM_Compiler_CompileShader_VectorVMPreprocessShader);
 
 	PRAGMA_DISABLE_DEPRECATION_WARNINGS		// FShaderCompilerDefinitions will be made internal in the future, marked deprecated until then
-
-	FString PreprocessedShader;
 	FShaderCompilerDefinitions AdditionalDefines;
-	EHlslCompileTarget HlslCompilerTarget = HCT_FeatureLevelSM5;
-	ECompilerFlags PlatformFlowControl = CFLAG_AvoidFlowControl;
-
 	AdditionalDefines.SetDefine(TEXT("COMPILER_HLSLCC"), 1);
 	AdditionalDefines.SetDefine(TEXT("COMPILER_VECTORVM"), 1);
 	AdditionalDefines.SetDefine(TEXT("VECTORVM_PROFILE"), 1);
-	
-	const bool bDumpDebugInfo = (Input.DumpDebugInfoPath != TEXT("") && IFileManager::Get().DirectoryExists(*Input.DumpDebugInfoPath));
-
 	AdditionalDefines.SetDefine(TEXT("FORCE_FLOATS"), (uint32)1);
-	
 	PRAGMA_ENABLE_DEPRECATION_WARNINGS
-
 	if (Input.bSkipPreprocessedCache)
 	{
-		if (!FFileHelper::LoadFileToString(PreprocessedShader, *Input.VirtualSourceFilePath))
+		if (!FFileHelper::LoadFileToString(Output.EditSource(), *Input.VirtualSourceFilePath))
 		{
 			return false;
 		}
 
 		// Remove const as we are on debug-only mode
-		CrossCompiler::CreateEnvironmentFromResourceTable(PreprocessedShader, (FShaderCompilerEnvironment&)Input.Environment);
+		CrossCompiler::CreateEnvironmentFromResourceTable(Output.GetSource(), (FShaderCompilerEnvironment&)Input.Environment);
 	}
 	else
 	{
-		SCOPE_CYCLE_COUNTER(STAT_VectorVM_Compiler_CompileShader_VectorVMPreprocessShader);
-
 		// Don't include shader definitions since it creates shader compilation errors.
-		if (!PreprocessShader(PreprocessedShader, Output, Input, AdditionalDefines, EDumpShaderDefines::DontIncludeDefines))
+		if (!PreprocessShader(Output, Input, Environment, AdditionalDefines, EDumpShaderDefines::DontIncludeDefines))
 		{
 			// The preprocessing stage will add any relevant errors.
-			if (Output.Errors.Num() != 0)
-			{
-				FString Errors;
-				for (const FShaderCompilerError& Error : Output.Errors)
-				{
-					Errors += Error.GetErrorString() + TEXT("\r\n");
-				}
-				VMCompilationOutput.Errors = Errors;
-			}
 			return false;
 		}
 	}
 
-	//TODO: Need to remove any unsupported features here?
+	return true;
+}
+
+bool CompileVectorVMShader(
+	const FShaderCompilerInput& Input,
+	const FShaderPreprocessOutput& PreprocessOutput,
+	FVectorVMCompilationOutput& VMCompilationOutput,
+	const FString& WorkingDirectory,
+	bool bSkipBackendOptimizations)
+{
+	SCOPE_CYCLE_COUNTER(STAT_VectorVM_Compiler_CompileShader_VectorVM);
+
+	EHlslCompileTarget HlslCompilerTarget = HCT_FeatureLevelSM5;
+	ECompilerFlags PlatformFlowControl = CFLAG_AvoidFlowControl;
 
 	char* ShaderSource = NULL;
 	char* ErrorLog = NULL;
 
 	const EHlslShaderFrequency Frequency = HSF_VertexShader;
 
-	//Is stuff like this needed? What others?
 	uint32 CCFlags = HLSLCC_NoPreprocess;
-		//CCFlags |= HLSLCC_PrintAST;
-	//CCFlags |= HLSLCC_UseFullPrecisionInPS;
 
 	if (bSkipBackendOptimizations)
 	{
 		CCFlags |= HLSLCC_DisableBackendOptimizations;
 	}
 
-	//TODO: Do this later when we implement the rest of the shader plumbing stuff.
-// 		if (bDumpDebugInfo)
-// 		{
-// 			const FString VVMFile = (Input.DumpDebugInfoPath / TEXT("Output.vvm"));
-// 			const FString USFFile = (Input.DumpDebugInfoPath / Input.GetSourceFilename());
-// 			const FString CCBatchFileContents = CreateCrossCompilerBatchFile(USFFile, VVMFile, *Input.EntryPointName, Frequency, Version, CCFlags);
-// 			if (!CCBatchFileContents.IsEmpty())
-// 			{
-// 				const TCHAR * ScriptName = PLATFORM_WINDOWS ? TEXT("CrossCompile.bat") : TEXT("CrossCompile.sh");
-// 				FFileHelper::SaveStringToFile(CCBatchFileContents, *(Input.DumpDebugInfoPath / ScriptName));
-// 			}
-// 		}
-
-	//NEEDED?
-	// Required as we added the RemoveUniformBuffersFromSource() function (the cross-compiler won't be able to interpret comments w/o a preprocessor)
-	//CCFlags &= ~HLSLCC_NoPreprocess;
-
-	UE::ShaderCompilerCommon::DumpDebugShaderData(Input, PreprocessedShader);
-
 	FVectorVMCodeBackend VVMBackEnd(CCFlags, HlslCompilerTarget, VMCompilationOutput);
-		FVectorVMLanguageSpec VVMLanguageSpec; 
+	FVectorVMLanguageSpec VVMLanguageSpec; 
 
 	bool bResult = false;
 	{
@@ -150,7 +96,7 @@ bool CompileShader_VectorVM(const FShaderCompilerInput& Input, FShaderCompilerOu
 			SCOPE_CYCLE_COUNTER(STAT_VectorVM_Compiler_CompileShader_CrossCompilerContextRun);
 
 			bResult = CrossCompilerContext.Run(
-				TCHAR_TO_ANSI(*PreprocessedShader),
+				TCHAR_TO_ANSI(*PreprocessOutput.GetSource()),
 				TCHAR_TO_ANSI(*Input.EntryPointName),
 				&VVMBackEnd,
 				&ShaderSource,
@@ -171,7 +117,7 @@ bool CompileShader_VectorVM(const FShaderCompilerInput& Input, FShaderCompilerOu
 		VMCompilationOutput.Errors = Converted.GetData();
 
 		TArray<FString> OutputByLines;
-		PreprocessedShader.ParseIntoArrayLines(OutputByLines, false);
+		PreprocessOutput.GetSource().ParseIntoArrayLines(OutputByLines, false);
 
 		UE_LOG(LogVectorVMShaderCompiler, Warning, TEXT("Warnings while processing %s"), *Input.DebugGroupName);
 
@@ -186,74 +132,6 @@ bool CompileShader_VectorVM(const FShaderCompilerInput& Input, FShaderCompilerOu
 		VMCompilationOutput.Errors.Empty();
 	}
 
-	//TODO: Try to get rid of the CompilationOutput and have the vm bytecode life in the shader eco-system as the compute shader version will.
-//		int32 SourceLen = VVMBackEnd.ByteCode.Num();//ShaderSource ? FCStringAnsi::Strlen(ShaderSource) : 0;
-// 		if (SourceLen > 0)
-// 		{
-// 			if (bDumpDebugInfo)
-// 			{
-// 				const FString GLSLFile = (Input.DumpDebugInfoPath / TEXT("Output.vvm"));
-// 
-// 				if (SourceLen > 0)
-// 				{
-// 					uint32 Len = FCStringAnsi::Strlen(TCHAR_TO_ANSI(*Input.VirtualSourceFilePath)) + FCStringAnsi::Strlen(TCHAR_TO_ANSI(*Input.EntryPointName)) + FCStringAnsi::Strlen(ShaderSource) + 20;
-// 					char* Dest = (char*)malloc(Len);
-// 					FCStringAnsi::Snprintf(Dest, Len, "// ! %s:%s\n%s", (const char*)TCHAR_TO_ANSI(*Input.VirtualSourceFilePath), (const char*)TCHAR_TO_ANSI(*Input.EntryPointName), (const char*)ShaderSource);
-// 					free(ShaderSource);
-// 					ShaderSource = Dest;
-// 					SourceLen = FCStringAnsi::Strlen(ShaderSource);
-// 					
-// 					FArchive* FileWriter = IFileManager::Get().CreateFileWriter(*(Input.DumpDebugInfoPath / Input.VirtualSourceFilePath + TEXT(".vvm")));
-// 					if (FileWriter)
-// 					{
-// 						FileWriter->Serialize(ShaderSource,SourceLen+1);
-// 						FileWriter->Close();
-// 						delete FileWriter;
-// 					}
-// 				}
-// 			}
-
-		//HMMMM....?
-// #if VALIDATE_GLSL_WITH_DRIVER
-// 			PrecompileShader(Output, Input, ShaderSource, Version, Frequency);
-// #else // VALIDATE_GLSL_WITH_DRIVER
-		//SourceLen = FCStringAnsi::Strlen(ShaderSource);
-//			Output.Target = Input.Target;
-		//BuildShaderOutput(Output, Input, ShaderSource, SourceLen, Version);
-
-//			Output.bSucceeded = true;
-		//Should we add optional data to define the register allocations etc so that someone can just write hlsl directly and have it become a niagara sim?
-//			Output.ShaderCode.GetWriteAccess().Append((uint8*)VVMBackEnd.ByteCode.GetData(), VVMBackEnd.ByteCode.Num());
-//			Output.ShaderCode.FinalizeShaderCode();
-//#endif // VALIDATE_GLSL_WITH_DRIVER
-// 		}
-// 		else
-// 		{
-		//ES2 Command line is throwing me off. Is this needed?
-
-// 			if (bDumpDebugInfo)
-// 			{
-// 				// Generate the batch file to help track down cross-compiler issues if necessary
-// 				const FString VVMFile = (Input.DumpDebugInfoPath / TEXT("Output.vvm"));
-// 				const FString VVMBatchFileContents = CreateCommandLineGLSLES2(VVMFile, (Input.DumpDebugInfoPath / TEXT("Output.asm")), Version, Frequency, false);
-// 				if (!VVMBatchFileContents.IsEmpty())
-// 				{
-// 					FFileHelper::SaveStringToFile(VVMBatchFileContents, *(Input.DumpDebugInfoPath / TEXT("VVMCompile.bat")));
-// 				}
-// 			}
-// 
-// 			Output.bSucceeded = false;
-// 
-// 			FString Tmp = ANSI_TO_TCHAR(ErrorLog);
-// 			TArray<FString> ErrorLines;
-// 			Tmp.ParseIntoArray(ErrorLines, TEXT("\n"), true);
-// 			for (int32 LineIndex = 0; LineIndex < ErrorLines.Num(); ++LineIndex)
-// 			{
-// 				const FString& Line = ErrorLines[LineIndex];
-// 				CrossCompiler::ParseHlslccError(Output.Errors, Line);
-// 			}
-// 		}
-
 	if (ShaderSource)
 	{
 		UE_LOG(LogVectorVMShaderCompiler, Warning, TEXT("%hs"), (const char*)ShaderSource);
@@ -268,3 +146,52 @@ bool CompileShader_VectorVM(const FShaderCompilerInput& Input, FShaderCompilerOu
 	return bResult;
 }
 
+bool CompileVectorVMShader(const FShaderCompilerInput& Input, const FShaderPreprocessOutput& PreprocessOutput, FShaderCompilerOutput& Output, const FString& WorkingDirectory)
+{
+	FVectorVMCompilationOutput CompilationOutput;
+	bool bResult = CompileVectorVMShader(Input, PreprocessOutput, CompilationOutput, WorkingDirectory, Input.Environment.CompilerFlags.Contains(CFLAG_SkipOptimizations));
+	
+	if (bResult)
+	{
+		FMemoryWriter Ar(Output.ShaderCode.GetWriteAccess(), true);
+		Ar << CompilationOutput;
+		Output.bSucceeded = true;
+	}
+	else if (CompilationOutput.Errors.Len() > 0)
+	{
+		Output.Errors.Add(FShaderCompilerError(*CompilationOutput.Errors));
+	}
+	return bResult;
+}
+
+void OutputVectorVMDebugData(
+	const FShaderCompilerInput& Input,
+	const FShaderPreprocessOutput& PreprocessOutput,
+	const FShaderCompilerOutput& Output)
+{
+	UE::ShaderCompilerCommon::DumpExtendedDebugShaderData(Input, PreprocessOutput, Output);
+}
+
+bool TestCompileVectorVMShader(
+	const FShaderCompilerInput& Input,
+	const FString& WorkingDirectory,
+	FVectorVMCompilationOutput& VMCompilationOutput,
+	bool bSkipBackendOptimizations)
+{
+	FShaderPreprocessOutput PreprocessOutput;
+	if (!PreprocessVectorVMShader(Input, Input.Environment, PreprocessOutput))
+	{
+		if (PreprocessOutput.GetErrors().Num() != 0)
+		{
+			FString Errors;
+			for (const FShaderCompilerError& Error : PreprocessOutput.GetErrors())
+			{
+				Errors += Error.GetErrorString() + TEXT("\r\n");
+			}
+			VMCompilationOutput.Errors = Errors;
+		}
+		return false;
+	}
+
+	return CompileVectorVMShader(Input, PreprocessOutput, VMCompilationOutput, WorkingDirectory, bSkipBackendOptimizations);
+}
diff --git a/Engine/Source/Developer/ShaderFormatVectorVM/Public/ShaderFormatVectorVM.h b/Engine/Source/Developer/ShaderFormatVectorVM/Public/ShaderFormatVectorVM.h
--- a/Engine/Source/Developer/ShaderFormatVectorVM/Public/ShaderFormatVectorVM.h
+++ b/Engine/Source/Developer/ShaderFormatVectorVM/Public/ShaderFormatVectorVM.h
@@ -9,6 +9,9 @@
 #include "VectorVM.h"
 #include "VectorVMCommon.h"
 
+struct FShaderCompilerInput;
+struct FShaderCompilerOutput;
+
 /** Ordered table of functions actually called by the VM script. */
 struct FCalledVMFunction
 {
@@ -61,7 +64,8 @@ inline FArchive& operator<<(FArchive& Ar, FVectorVMCompilationOutput& Output)
 	return Ar;
 }
 
-extern bool SHADERFORMATVECTORVM_API CompileShader_VectorVM(const struct FShaderCompilerInput& Input, struct FShaderCompilerOutput& Output, const class FString& WorkingDirectory, uint8 Version);
+UE_DEPRECATED(5.4, "Direct VectorVM compilation will no longer be exposed outside of the TestCompileVectorVMShader function")
+inline bool CompileShader_VectorVM(const FShaderCompilerInput& Input, FShaderCompilerOutput& Output, const FString& WorkingDirectory, uint8 Version) { return false; }
 
-//Cheating hack version. To be removed when we add all the plumbing for VVM scripts to be treat like real shaders.
-extern bool SHADERFORMATVECTORVM_API CompileShader_VectorVM(const struct FShaderCompilerInput& Input,struct FShaderCompilerOutput& Output,const class FString& WorkingDirectory, uint8 Version, struct FVectorVMCompilationOutput& VMCompilationOutput, bool bSkipBackendOptimizations = false);
+UE_DEPRECATED(5.4, "This overload of CompileShader_VectorVM is superceded by TestCompileShaderVectorVM")
+inline bool CompileShader_VectorVM(const FShaderCompilerInput& Input, FShaderCompilerOutput& Output, const FString& WorkingDirectory, uint8 Version, FVectorVMCompilationOutput& VMCompilationOutput, bool bSkipBackendOptimizations = false) { return false; }
diff --git a/Engine/Source/Runtime/RenderCore/Public/ShaderPreprocessTypes.h b/Engine/Source/Runtime/RenderCore/Public/ShaderPreprocessTypes.h
--- a/Engine/Source/Runtime/RenderCore/Public/ShaderPreprocessTypes.h
+++ b/Engine/Source/Runtime/RenderCore/Public/ShaderPreprocessTypes.h
@@ -135,6 +135,11 @@ public:
 		return Errors;
 	}
 
+	inline TConstArrayView<FShaderCompilerError> GetErrors() const
+	{
+		return MakeArrayView(Errors);
+	}
+
 	inline bool ParseAndModify(
 		const FShaderCompilerInput& Input, 
 		const FShaderCompilerEnvironment& Environment, 
