Fix warning V1052: Declaring virtual methods in a class marked as 'final' is pointless.

diff --git a/Engine/Plugins/Experimental/ImagePlate/Source/ImagePlate/Private/ImagePlateComponent.cpp b/Engine/Plugins/Experimental/ImagePlate/Source/ImagePlate/Private/ImagePlateComponent.cpp
--- a/Engine/Plugins/Experimental/ImagePlate/Source/ImagePlate/Private/ImagePlateComponent.cpp
+++ b/Engine/Plugins/Experimental/ImagePlate/Source/ImagePlate/Private/ImagePlateComponent.cpp
@@ -176,7 +176,7 @@ namespace
 			}
 		}
 
-		virtual void GetDynamicMeshElements(const TArray<const FSceneView*>& Views, const FSceneViewFamily& ViewFamily, uint32 VisibilityMap, FMeshElementCollector& Collector) const
+		virtual void GetDynamicMeshElements(const TArray<const FSceneView*>& Views, const FSceneViewFamily& ViewFamily, uint32 VisibilityMap, FMeshElementCollector& Collector) const override
 		{
 			QUICK_SCOPE_CYCLE_COUNTER( STAT_ImagePlateSceneProxy_GetDynamicMeshElements );
 
diff --git a/Engine/Source/Developer/ShaderFormatVectorVM/Private/ir_vm_merge_op_visitor.cpp b/Engine/Source/Developer/ShaderFormatVectorVM/Private/ir_vm_merge_op_visitor.cpp
--- a/Engine/Source/Developer/ShaderFormatVectorVM/Private/ir_vm_merge_op_visitor.cpp
+++ b/Engine/Source/Developer/ShaderFormatVectorVM/Private/ir_vm_merge_op_visitor.cpp
@@ -51,13 +51,13 @@ public:
 	{
 	}
 
-	virtual ir_visitor_status visit_enter(ir_expression* expr)
+	virtual ir_visitor_status visit_enter(ir_expression* expr) override
 	{
 		++expr_depth;
 		check(expr_depth == 1);//Should never exceed one expression.
 		return visit_continue;
 	}
-	virtual ir_visitor_status visit_leave(ir_expression* expr)
+	virtual ir_visitor_status visit_leave(ir_expression* expr) override
 	{
 		if (expr->operation == ir_binop_mul)
 		{
@@ -163,7 +163,7 @@ public:
 		return visit_continue;
 	}
 
-	virtual ir_visitor_status visit_enter(ir_assignment* assign)
+	virtual ir_visitor_status visit_enter(ir_assignment* assign) override
 	{
 		assign_array_to_add = nullptr;
 		expr_depth = 0;
diff --git a/Engine/Source/Developer/ShaderFormatVectorVM/Private/ir_vm_propagate_non_expressions_visitor.cpp b/Engine/Source/Developer/ShaderFormatVectorVM/Private/ir_vm_propagate_non_expressions_visitor.cpp
--- a/Engine/Source/Developer/ShaderFormatVectorVM/Private/ir_vm_propagate_non_expressions_visitor.cpp
+++ b/Engine/Source/Developer/ShaderFormatVectorVM/Private/ir_vm_propagate_non_expressions_visitor.cpp
@@ -198,13 +198,13 @@ public:
 		}
 	}
 
-	virtual ir_visitor_status visit_enter(ir_expression*)
+	virtual ir_visitor_status visit_enter(ir_expression*) override
 	{
 		num_expr++;
 		return visit_continue;
 	}
 
-	virtual ir_visitor_status visit_enter(ir_assignment* assign)
+	virtual ir_visitor_status visit_enter(ir_assignment* assign) override
 	{
 		if (assign->condition)
 		{
@@ -348,7 +348,7 @@ public:
 	{
 	}
 
-	virtual ir_visitor_status visit_enter(ir_call* call)
+	virtual ir_visitor_status visit_enter(ir_call* call) override
 	{
 		bool bFoundMatrixParam = false;
 		foreach_iter(exec_list_iterator, iter, call->actual_parameters)
@@ -690,7 +690,7 @@ public:
 		}
 	}
 
-	virtual ir_visitor_status visit_enter(ir_assignment* assign)
+	virtual ir_visitor_status visit_enter(ir_assignment* assign) override
 	{
 		if (assign->condition)
 		{
diff --git a/Engine/Source/Developer/ShaderFormatVectorVM/Private/ir_vm_to_single_op_visitor.cpp b/Engine/Source/Developer/ShaderFormatVectorVM/Private/ir_vm_to_single_op_visitor.cpp
--- a/Engine/Source/Developer/ShaderFormatVectorVM/Private/ir_vm_to_single_op_visitor.cpp
+++ b/Engine/Source/Developer/ShaderFormatVectorVM/Private/ir_vm_to_single_op_visitor.cpp
@@ -45,7 +45,7 @@ public:
 	{
 	}
 
-	virtual ir_visitor_status visit_enter(ir_swizzle* swiz)
+	virtual ir_visitor_status visit_enter(ir_swizzle* swiz) override
 	{
 		expression_stack.Push(swiz);
 		swiz->val->accept(this);
@@ -61,7 +61,7 @@ public:
 		return visit_continue_with_parent;
 	}
 
-	virtual ir_visitor_status visit_enter(ir_expression* expr)
+	virtual ir_visitor_status visit_enter(ir_expression* expr) override
 	{
 		check(base_ir->next && base_ir->prev);
 		assign_has_expressions = true;
@@ -96,7 +96,7 @@ public:
 		return visit_continue_with_parent;
 	}
 
-	virtual ir_visitor_status visit_enter(ir_assignment* assign)
+	virtual ir_visitor_status visit_enter(ir_assignment* assign) override
 	{
 		check(base_ir->next && base_ir->prev);
 		assign_has_expressions = false;
@@ -118,7 +118,7 @@ public:
 		return visit_continue_with_parent;
 	}
 
-	virtual ir_visitor_status visit_enter(ir_call* call)
+	virtual ir_visitor_status visit_enter(ir_call* call) override
 	{
 		check(base_ir->next && base_ir->prev);
 		assign_has_expressions = true;
diff --git a/Engine/Source/Editor/UnrealEd/Private/Cooker/IoStoreCookOnTheFlyRequestManager.cpp b/Engine/Source/Editor/UnrealEd/Private/Cooker/IoStoreCookOnTheFlyRequestManager.cpp
--- a/Engine/Source/Editor/UnrealEd/Private/Cooker/IoStoreCookOnTheFlyRequestManager.cpp
+++ b/Engine/Source/Editor/UnrealEd/Private/Cooker/IoStoreCookOnTheFlyRequestManager.cpp
@@ -404,7 +404,7 @@ private:
 	{
 	}
 
-	virtual void OnPackageGenerated(const FName& PackageName)
+	virtual void OnPackageGenerated(const FName& PackageName) override
 	{
 		FPackageId PackageId = FPackageId::FromName(PackageName);
 		UE_LOG(LogCookOnTheFly, Verbose, TEXT("Package 0x%llX '%s' generated"), PackageId.ValueForDebugging(), *PackageName.ToString());
diff --git a/Engine/Source/Runtime/Experimental/Chaos/Public/Chaos/AABBTree.h b/Engine/Source/Runtime/Experimental/Chaos/Public/Chaos/AABBTree.h
--- a/Engine/Source/Runtime/Experimental/Chaos/Public/Chaos/AABBTree.h
+++ b/Engine/Source/Runtime/Experimental/Chaos/Public/Chaos/AABBTree.h
@@ -1770,7 +1770,7 @@ public:
 		DeAllocateLeafNode(LeafNodeIdx);
 	}
 
-	virtual bool RemoveElement(const TPayloadType& Payload)
+	virtual bool RemoveElement(const TPayloadType& Payload) override
 	{
 		if (UNLIKELY(!ensure(bModifyingTreeMultiThreadingFastCheck == false)))
 		{
diff --git a/Engine/Source/Runtime/OpenGLDrv/Public/OpenGLDrv.h b/Engine/Source/Runtime/OpenGLDrv/Public/OpenGLDrv.h
--- a/Engine/Source/Runtime/OpenGLDrv/Public/OpenGLDrv.h
+++ b/Engine/Source/Runtime/OpenGLDrv/Public/OpenGLDrv.h
@@ -321,9 +321,9 @@ public:
 #endif
 
 	// FDynamicRHI interface.
-	virtual void Init();
+	virtual void Init() override;
 
-	virtual void Shutdown();
+	virtual void Shutdown() override;
 	virtual const TCHAR* GetName() override { return TEXT("OpenGL"); }
 
 	template<typename TRHIType>
@@ -763,7 +763,7 @@ public:
 
 	virtual FSamplerStateRHIRef RHICreateSamplerState(const FSamplerStateInitializerRHI& Initializer) final override;
 
-	virtual FGraphicsPipelineStateRHIRef RHICreateGraphicsPipelineState(const FGraphicsPipelineStateInitializer& Initializer)
+	virtual FGraphicsPipelineStateRHIRef RHICreateGraphicsPipelineState(const FGraphicsPipelineStateInitializer& Initializer) override
 	{
 		PrepareGFXBoundShaderState(Initializer);
 
