- Fix not handling constant matrix case in VM shader format

diff --git a/Engine/Source/Developer/ShaderFormatVectorVM/Private/ir_vm_propagate_non_expressions_visitor.cpp b/Engine/Source/Developer/ShaderFormatVectorVM/Private/ir_vm_propagate_non_expressions_visitor.cpp
--- a/Engine/Source/Developer/ShaderFormatVectorVM/Private/ir_vm_propagate_non_expressions_visitor.cpp
+++ b/Engine/Source/Developer/ShaderFormatVectorVM/Private/ir_vm_propagate_non_expressions_visitor.cpp
@@ -376,11 +376,21 @@ public:
 		{
 			if (rval->type->is_matrix())
 			{
-				MatrixVectors& mv = MatrixVectorMap.FindChecked(rval->variable_referenced());
-				call->actual_parameters.push_tail(new(parse_state) ir_dereference_variable(mv.v[0]));
-				call->actual_parameters.push_tail(new(parse_state) ir_dereference_variable(mv.v[1]));
-				call->actual_parameters.push_tail(new(parse_state) ir_dereference_variable(mv.v[2]));
-				call->actual_parameters.push_tail(new(parse_state) ir_dereference_variable(mv.v[3]));
+				if (rval->variable_referenced()->mode == ir_var_uniform)
+				{
+					call->actual_parameters.push_tail(new(parse_state) ir_dereference_array(rval->variable_referenced(), new(parse_state) ir_constant(0)));
+					call->actual_parameters.push_tail(new(parse_state) ir_dereference_array(rval->variable_referenced(), new(parse_state) ir_constant(1)));
+					call->actual_parameters.push_tail(new(parse_state) ir_dereference_array(rval->variable_referenced(), new(parse_state) ir_constant(2)));
+					call->actual_parameters.push_tail(new(parse_state) ir_dereference_array(rval->variable_referenced(), new(parse_state) ir_constant(3)));
+				}
+				else
+				{
+					MatrixVectors& mv = MatrixVectorMap.FindChecked(rval->variable_referenced());
+					call->actual_parameters.push_tail(new(parse_state) ir_dereference_variable(mv.v[0]));
+					call->actual_parameters.push_tail(new(parse_state) ir_dereference_variable(mv.v[1]));
+					call->actual_parameters.push_tail(new(parse_state) ir_dereference_variable(mv.v[2]));
+					call->actual_parameters.push_tail(new(parse_state) ir_dereference_variable(mv.v[3]));
+				}
 			}
 			else
 			{
